[MLDeformer] Fixed some issue where multiple ML Deformers that would point to the same Skeletal Mesh but trained with different settings, would corrupt each other. Moved the external morph sets from the skeletal mesh render data to the skinned mesh component instead in order to fix this. Also fixed some crashes.

diff --git a/Engine/Plugins/Experimental/MLDeformer/MLDeformerFramework/Source/MLDeformerFramework/Private/MLDeformerComponent.cpp b/Engine/Plugins/Experimental/MLDeformer/MLDeformerFramework/Source/MLDeformerFramework/Private/MLDeformerComponent.cpp
--- a/Engine/Plugins/Experimental/MLDeformer/MLDeformerFramework/Source/MLDeformerFramework/Private/MLDeformerComponent.cpp
+++ b/Engine/Plugins/Experimental/MLDeformer/MLDeformerFramework/Source/MLDeformerFramework/Private/MLDeformerComponent.cpp
@@ -84,9 +84,9 @@ void UMLDeformerComponent::AddNeuralNetworkModifyDelegate()
 
 void UMLDeformerComponent::RemoveNeuralNetworkModifyDelegate()
 {
-	if (DeformerAsset != nullptr && 
+	if (DeformerAsset && 
 		NeuralNetworkModifyDelegateHandle != FDelegateHandle() && 
-		DeformerAsset->GetModel() != nullptr)
+		DeformerAsset->GetModel())
 	{
 		DeformerAsset->GetModel()->NeuralNetworkModifyDelegate.Remove(NeuralNetworkModifyDelegateHandle);
 	}
diff --git a/Engine/Plugins/Experimental/MLDeformer/MLDeformerFramework/Source/MLDeformerFramework/Private/MLDeformerGraphDataInterface.cpp b/Engine/Plugins/Experimental/MLDeformer/MLDeformerFramework/Source/MLDeformerFramework/Private/MLDeformerGraphDataInterface.cpp
--- a/Engine/Plugins/Experimental/MLDeformer/MLDeformerFramework/Source/MLDeformerFramework/Private/MLDeformerGraphDataInterface.cpp
+++ b/Engine/Plugins/Experimental/MLDeformer/MLDeformerFramework/Source/MLDeformerFramework/Private/MLDeformerGraphDataInterface.cpp
@@ -70,15 +70,18 @@ namespace UE::MLDeformer
 		using namespace UE::MLDeformer;
 
 		const UMLDeformerAsset* DeformerAsset = DeformerComponent->GetDeformerAsset();
-		const UMLDeformerModel* Model = DeformerAsset->GetModel();
-		const UMLDeformerModelInstance* ModelInstance = DeformerComponent->GetModelInstance();
+		if (DeformerAsset)
+		{
+			const UMLDeformerModel* Model = DeformerAsset->GetModel();
+			const UMLDeformerModelInstance* ModelInstance = DeformerComponent->GetModelInstance();
 		
-		SkeletalMeshObject = ModelInstance->GetSkeletalMeshComponent()->MeshObject;
-		NeuralNetwork = Model->GetNeuralNetwork();
-		NeuralNetworkInferenceHandle = ModelInstance->GetNeuralNetworkInferenceHandle();
-		bCanRunNeuralNet = ModelInstance->IsCompatible();
-		Weight = DeformerComponent->GetWeight();
-		VertexMapBufferSRV = Model->GetVertexMapBuffer().ShaderResourceViewRHI;
+			SkeletalMeshObject = ModelInstance->GetSkeletalMeshComponent()->MeshObject;
+			NeuralNetwork = Model->GetNeuralNetwork();
+			NeuralNetworkInferenceHandle = ModelInstance->GetNeuralNetworkInferenceHandle();
+			bCanRunNeuralNet = ModelInstance->IsCompatible();
+			Weight = DeformerComponent->GetWeight();
+			VertexMapBufferSRV = Model->GetVertexMapBuffer().ShaderResourceViewRHI;
+		}
 	}
 
 	void FMLDeformerGraphDataProviderProxy::AllocateResources(FRDGBuilder& GraphBuilder)
diff --git a/Engine/Plugins/Experimental/MLDeformer/MLDeformerFramework/Source/MLDeformerFramework/Private/MLDeformerGraphDebugDataInterface.cpp b/Engine/Plugins/Experimental/MLDeformer/MLDeformerFramework/Source/MLDeformerFramework/Private/MLDeformerGraphDebugDataInterface.cpp
--- a/Engine/Plugins/Experimental/MLDeformer/MLDeformerFramework/Source/MLDeformerFramework/Private/MLDeformerGraphDebugDataInterface.cpp
+++ b/Engine/Plugins/Experimental/MLDeformer/MLDeformerFramework/Source/MLDeformerFramework/Private/MLDeformerGraphDebugDataInterface.cpp
@@ -95,15 +95,18 @@ namespace UE::MLDeformer
 	{
 		Provider = InProvider;
 
-		UMLDeformerModel* Model = DeformerAsset->GetModel();	
-		UMLDeformerVizSettings* VizSettings = Model->GetVizSettings();
-		const UMLDeformerModelInstance* ModelInstance = DeformerComponent->GetModelInstance();
-
-		SkeletalMeshObject = ModelInstance->GetSkeletalMeshComponent()->MeshObject;
-		VertexMapBufferSRV = Model->GetVertexMapBuffer().ShaderResourceViewRHI;
-		HeatMapMode = (int32)VizSettings->GetHeatMapMode();
-		HeatMapMax = 1.0f / FMath::Max(VizSettings->GetHeatMapMax(), 0.00001f);
-		GroundTruthLerp = VizSettings->GetGroundTruthLerp();
+		if (DeformerAsset)
+		{
+			UMLDeformerModel* Model = DeformerAsset->GetModel();	
+			UMLDeformerVizSettings* VizSettings = Model->GetVizSettings();
+			const UMLDeformerModelInstance* ModelInstance = DeformerComponent->GetModelInstance();
+
+			SkeletalMeshObject = ModelInstance->GetSkeletalMeshComponent()->MeshObject;
+			VertexMapBufferSRV = Model->GetVertexMapBuffer().ShaderResourceViewRHI;
+			HeatMapMode = (int32)VizSettings->GetHeatMapMode();
+			HeatMapMax = 1.0f / FMath::Max(VizSettings->GetHeatMapMax(), 0.00001f);
+			GroundTruthLerp = VizSettings->GetGroundTruthLerp();
+		}
 	}
 
 	void FMLDeformerGraphDebugDataProviderProxy::HandleZeroGroundTruthPositions()
diff --git a/Engine/Plugins/Experimental/MLDeformer/MLDeformerFramework/Source/MLDeformerFramework/Private/MLDeformerModelInstance.cpp b/Engine/Plugins/Experimental/MLDeformer/MLDeformerFramework/Source/MLDeformerFramework/Private/MLDeformerModelInstance.cpp
--- a/Engine/Plugins/Experimental/MLDeformer/MLDeformerFramework/Source/MLDeformerFramework/Private/MLDeformerModelInstance.cpp
+++ b/Engine/Plugins/Experimental/MLDeformer/MLDeformerFramework/Source/MLDeformerFramework/Private/MLDeformerModelInstance.cpp
@@ -17,8 +17,9 @@ void UMLDeformerModelInstance::BeginDestroy()
 
 void UMLDeformerModelInstance::Release()
 {
+	// Destroy the neural network instance.
 	UNeuralNetwork* NeuralNetwork = Model.Get() ? Model->GetNeuralNetwork() : nullptr;
-	if (NeuralNetwork != nullptr && NeuralNetworkInferenceHandle != -1)
+	if (NeuralNetwork && NeuralNetworkInferenceHandle != -1)
 	{
 		NeuralNetwork->DestroyInferenceContext(NeuralNetworkInferenceHandle);
 		NeuralNetworkInferenceHandle = -1;
@@ -84,8 +85,9 @@ FString UMLDeformerModelInstance::CheckCompatibility(USkeletalMeshComponent* InS
 		}
 	}
 
+	// Verify the number of inputs versus the expected number of inputs.
 	UNeuralNetwork* NeuralNetwork = Model->GetNeuralNetwork();
-	if (NeuralNetwork != nullptr && NeuralNetwork->IsLoaded() && Model->GetDeformerAsset())
+	if (NeuralNetwork && NeuralNetwork->IsLoaded() && Model->GetDeformerAsset())
 	{
 		const int64 NumNeuralNetInputs = NeuralNetwork->GetInputTensor().Num();
 		const int64 NumDeformerAssetInputs = static_cast<int64>(Model->GetInputInfo()->CalcNumNeuralNetInputs());
@@ -131,6 +133,7 @@ FString UMLDeformerModelInstance::CheckCompatibility(USkeletalMeshComponent* InS
 
 void UMLDeformerModelInstance::UpdateBoneTransforms()
 {
+	// If we use a leader component, we have to get the transforms from that component.
 	const USkinnedMeshComponent* LeaderPoseComponent = SkeletalMeshComponent->LeaderPoseComponent.Get();
 	if (LeaderPoseComponent)
 	{
@@ -155,6 +158,8 @@ void UMLDeformerModelInstance::UpdateBoneTransforms()
 	}
 	else
 	{
+		// Grab the transforms from our own skeletal mesh component.
+		// These are local space transforms, relative to the parent bone.
 		BoneTransforms = SkeletalMeshComponent->GetBoneSpaceTransforms();
 		const int32 NumTrainingBones = AssetBonesToSkelMeshMappings.Num();
 		for (int32 Index = 0; Index < NumTrainingBones; ++Index)
@@ -266,14 +271,14 @@ void UMLDeformerModelInstance::RunNeuralNetwork(float ModelWeight)
 		// NOTE: Inputs still come from the CPU.
 		check(NeuralNetwork->GetDeviceType() == ENeuralDeviceType::GPU && NeuralNetwork->GetInputDeviceType() == ENeuralDeviceType::CPU && NeuralNetwork->GetOutputDeviceType() == ENeuralDeviceType::GPU);
 		ENQUEUE_RENDER_COMMAND(RunNeuralNetwork)
-			(
-				[NeuralNetwork, Handle = NeuralNetworkInferenceHandle](FRHICommandListImmediate& RHICmdList)
-				{
-					// Output deltas will be available on GPU for DeformerGraph via UMLDeformerDataProvider.
-					FRDGBuilder GraphBuilder(RHICmdList);
-					NeuralNetwork->Run(GraphBuilder, Handle);
-					GraphBuilder.Execute();
-				}
+		(
+			[NeuralNetwork, Handle = NeuralNetworkInferenceHandle](FRHICommandListImmediate& RHICmdList)
+			{
+				// Output deltas will be available on GPU for DeformerGraph via UMLDeformerDataProvider.
+				FRDGBuilder GraphBuilder(RHICmdList);
+				NeuralNetwork->Run(GraphBuilder, Handle);
+				GraphBuilder.Execute();
+			}
 		);
 	}
 	else
@@ -330,13 +335,21 @@ bool UMLDeformerModelInstance::SetupNeuralNetworkForFrame()
 	return true;
 }
 
-
 void UMLDeformerModelInstance::Tick(float DeltaTime, float ModelWeight)
 {
+	// Post init hasn't yet succeeded, try again.
+	// This could for example happen when you add an ML Deformer component, but your SkeletalMesh isn't setup yet, but later becomes valid.
+	if (Model && !HasPostInitialized())
+	{
+		Model->PostMLDeformerComponentInit(this);
+	}
+
+	// Setup the neural network, by feeding it's inputs and doing some compatibility checks.
 	if (!SetupNeuralNetworkForFrame())
 	{
 		return;
 	}
 
+	// Perform the inference, calculate the network outputs and possibly use them, depending on how the model works.
 	RunNeuralNetwork(ModelWeight);
 }
diff --git a/Engine/Plugins/Experimental/MLDeformer/MLDeformerFramework/Source/MLDeformerFramework/Public/MLDeformerGraphDebugDataInterface.h b/Engine/Plugins/Experimental/MLDeformer/MLDeformerFramework/Source/MLDeformerFramework/Public/MLDeformerGraphDebugDataInterface.h
--- a/Engine/Plugins/Experimental/MLDeformer/MLDeformerFramework/Source/MLDeformerFramework/Public/MLDeformerGraphDebugDataInterface.h
+++ b/Engine/Plugins/Experimental/MLDeformer/MLDeformerFramework/Source/MLDeformerFramework/Public/MLDeformerGraphDebugDataInterface.h
@@ -73,11 +73,16 @@ class USkeletalMeshComponent;
 	{ \
 		MLDEFORMER_EDITORDATA_ONLY( \
 			DataProviderProxyClassName* Proxy = new DataProviderProxyClassName(DeformerComponent, DeformerAsset, this); \
-			const float SampleTime = DeformerComponent->GetModelInstance()->GetSkeletalMeshComponent()->GetPosition(); \
-			UMLDeformerModel* Model = DeformerAsset->GetModel(); \
-			Model->SampleGroundTruthPositions(SampleTime, Proxy->GetGroundTruthPositions()); \
-			Proxy->HandleZeroGroundTruthPositions(); \
-			return Proxy;, \
+			if (DeformerAsset) \
+			{ \
+				const float SampleTime = DeformerComponent->GetModelInstance()->GetSkeletalMeshComponent()->GetPosition(); \
+				UMLDeformerModel* Model = DeformerAsset->GetModel(); \
+				Model->SampleGroundTruthPositions(SampleTime, Proxy->GetGroundTruthPositions()); \
+				Proxy->HandleZeroGroundTruthPositions(); \
+				return Proxy; \
+			} \
+			return nullptr; \
+			, \
 			return nullptr; \
 		) \
 	}
diff --git a/Engine/Plugins/Experimental/MLDeformer/MLDeformerFramework/Source/MLDeformerFramework/Public/MLDeformerModelInstance.h b/Engine/Plugins/Experimental/MLDeformer/MLDeformerFramework/Source/MLDeformerFramework/Public/MLDeformerModelInstance.h
--- a/Engine/Plugins/Experimental/MLDeformer/MLDeformerFramework/Source/MLDeformerFramework/Public/MLDeformerModelInstance.h
+++ b/Engine/Plugins/Experimental/MLDeformer/MLDeformerFramework/Source/MLDeformerFramework/Public/MLDeformerModelInstance.h
@@ -131,6 +131,11 @@ public:
 	 */
 	int32 GetNeuralNetworkInferenceHandle() const { return NeuralNetworkInferenceHandle; }
 
+	/** Set whether we had a successful call to UMLDeformerModel::PostMLDeformerComponentInit. */
+	void SetHasPostInitialized(bool bHasInitialized) { bHasPostInitialized = bHasInitialized; }
+
+	bool HasPostInitialized() const { return bHasPostInitialized; }
+
 protected:
 	/**
 	 * Set the bone transformations inside a given output buffer, starting from a given StartIndex.
@@ -181,4 +186,7 @@ protected:
 
 	/** Are the deformer asset and the used skeletal mesh component compatible? */
 	bool bIsCompatible = false;
+
+	/** Has this instance had a successful UMLDeformerModel::PostMLDeformerComponentInit call? */
+	bool bHasPostInitialized = false;
 };
diff --git a/Engine/Plugins/Experimental/MLDeformer/NeuralMorphModel/Source/NeuralMorphModel/Private/NeuralMorphModel.cpp b/Engine/Plugins/Experimental/MLDeformer/NeuralMorphModel/Source/NeuralMorphModel/Private/NeuralMorphModel.cpp
--- a/Engine/Plugins/Experimental/MLDeformer/NeuralMorphModel/Source/NeuralMorphModel/Private/NeuralMorphModel.cpp
+++ b/Engine/Plugins/Experimental/MLDeformer/NeuralMorphModel/Source/NeuralMorphModel/Private/NeuralMorphModel.cpp
@@ -7,8 +7,10 @@
 #include "MLDeformerComponent.h"
 #include "UObject/Object.h"
 #include "UObject/UObjectGlobals.h"
-#include "Rendering/SkeletalMeshLODRenderData.h"
+#include "Components/SkinnedMeshComponent.h"
 #include "Components/SkeletalMeshComponent.h"
+#include "GeometryCache.h"
+#include "Rendering/MorphTargetVertexInfoBuffers.h"
 
 #define LOCTEXT_NAMESPACE "NeuralMorphModel"
 
@@ -22,13 +24,14 @@ UNeuralMorphModel::UNeuralMorphModel(const FObjectInitializer& ObjectInitializer
 	VizSettings = ObjectInitializer.CreateEditorOnlyDefaultSubobject<UNeuralMorphModelVizSettings>(this, TEXT("VizSettings"));
 #endif
 
-	MorphTargetSet = MakeShared<FExternalMorphTargetSet>();
+	MorphTargetSet = MakeShared<FExternalMorphSet>();
 	MorphTargetSet->Name = FName(TEXT("NeuralBlendShapes"));
 }
 
 void UNeuralMorphModel::Serialize(FArchive& Archive)
 {
 	TRACE_CPUPROFILER_EVENT_SCOPE(UNeuralMorphModel::Serialize)
+
 	Super::Serialize(Archive);
 
 	// Check if we have initialized our compressed morph buffers.
@@ -46,11 +49,6 @@ void UNeuralMorphModel::Serialize(FArchive& Archive)
 		if (Archive.IsLoading())
 		{
 			Archive << MorphTargetSet->MorphBuffers;
-
-			// When we're in editor mode, keep the CPU data around, so we can re-initialize when needed.
-			#if WITH_EDITOR
-				MorphTargetSet->MorphBuffers.SetEmptyMorphCPUDataOnInitRHI(false);
-			#endif
 		}
 		else
 		{
@@ -103,43 +101,53 @@ UMLDeformerModelInstance* UNeuralMorphModel::CreateModelInstance(UMLDeformerComp
 
 void UNeuralMorphModel::PostMLDeformerComponentInit(UMLDeformerModelInstance* ModelInstance)
 {
+	if (ModelInstance->HasPostInitialized())
+	{
+		return;
+	}
+
 	TRACE_CPUPROFILER_EVENT_SCOPE(UNeuralMorphModel::PostMLDeformerComponentInit)
 
 	Super::PostMLDeformerComponentInit(ModelInstance);
 
 	// Register the external morph targets buffer to the render data of the skeletal mesh.
 	USkeletalMeshComponent* SkelMeshComponent = ModelInstance->GetSkeletalMeshComponent();
-	if (SkelMeshComponent)
+	if (SkelMeshComponent && SkelMeshComponent->GetSkeletalMesh())
 	{
-		FSkeletalMeshRenderData* RenderData = SkelMeshComponent->GetSkeletalMeshRenderData();
+		// Register the morph set. This overwrites the existing one for this model, if it already exists.
+		// Only add to LOD 0 for now.
+		const int32 LOD = 0;
+		SkelMeshComponent->AddExternalMorphSet(LOD, UNeuralMorphModel::NeuralMorphsExternalMorphSetID, MorphTargetSet);
+
+		// When we're in editor mode, keep the CPU data around, so we can re-initialize when needed.
+		#if WITH_EDITOR
+			MorphTargetSet->MorphBuffers.SetEmptyMorphCPUDataOnInitRHI(false);
+		#else
+			MorphTargetSet->MorphBuffers.SetEmptyMorphCPUDataOnInitRHI(true);
+		#endif
+
+		// Release the render resources, but only in an editor build.
+		// The non-editor build shouldn't do this, as then it can't initialize again. The non-editor build assumes
+		// that the data doesn't change and we don't need to re-init.
+		// In the editor build we have to re-initialize the render resources as the morph targets can change after (re)training, so
+		// that is why we release them here, and intialize them again after.
 		FMorphTargetVertexInfoBuffers& MorphBuffers = MorphTargetSet->MorphBuffers;
-		if (RenderData)
+		#if WITH_EDITOR
+			BeginReleaseResource(&MorphBuffers);
+		#endif
+
+		// Reinitialize the GPU compressed buffers.
+		if (MorphBuffers.IsMorphCPUDataValid() && MorphBuffers.GetNumMorphs() > 0)
 		{
-			// Register the morph set. This overwrites the existing one for this model, if it already exists.
-			// Only add to LOD 0 for now.
-			const int32 LOD = 0;
-			RenderData->LODRenderData[LOD].AddExternalMorphSet(UNeuralMorphModel::NeuralMorphsExternalMorphSetID, MorphTargetSet);
-
-			// Release the render resources, but only in an editor build.
-			// The non-editor build shouldn't do this, as then it can't initialize again. The non-editor build assumes
-			// that the data doesn't change and we don't need to re-init.
-			// In the editor build we have to re-initialize the render resources as the morph targets can change after (re)training, so
-			// that is why we release them here, and intialize them again after.
-			#if WITH_EDITOR
-				BeginReleaseResource(&MorphBuffers);
-			#endif
-
-			// Reinitialize the GPU compressed buffers.
-			if (MorphBuffers.IsMorphCPUDataValid() && MorphBuffers.GetNumMorphs() > 0)
-			{
-				// In a non-editor build this will clear the CPU data.
-				// That also means it can't re-init the resources later on again.
-				BeginInitResource(&MorphBuffers);
-			}
+			// In a non-editor build this will clear the CPU data.
+			// That also means it can't re-init the resources later on again.
+			BeginInitResource(&MorphBuffers);
 		}
 
 		// Update the weight information in the Skeletal Mesh.
 		SkelMeshComponent->RefreshExternalMorphTargetWeights();
+
+		ModelInstance->SetHasPostInitialized(true);
 	}
 }
 
@@ -163,6 +171,7 @@ void UNeuralMorphModel::BeginDestroy()
 	if (MorphTargetSet.IsValid())
 	{
 		BeginReleaseResource(&MorphTargetSet->MorphBuffers);
+		MorphTargetSet.Reset();
 	}
 	Super::BeginDestroy();
 }
diff --git a/Engine/Plugins/Experimental/MLDeformer/NeuralMorphModel/Source/NeuralMorphModel/Public/NeuralMorphModel.h b/Engine/Plugins/Experimental/MLDeformer/NeuralMorphModel/Source/NeuralMorphModel/Public/NeuralMorphModel.h
--- a/Engine/Plugins/Experimental/MLDeformer/NeuralMorphModel/Source/NeuralMorphModel/Public/NeuralMorphModel.h
+++ b/Engine/Plugins/Experimental/MLDeformer/NeuralMorphModel/Source/NeuralMorphModel/Public/NeuralMorphModel.h
@@ -17,7 +17,7 @@ class UMLDeformerAsset;
 class UMLDeformerModelInstance;
 class USkeleton;
 class IPropertyHandle;
-struct FExternalMorphTargetSet;
+struct FExternalMorphSet;
 
 UCLASS()
 class NEURALMORPHMODEL_API UNeuralMorphModel 
@@ -69,7 +69,7 @@ public:
 
 public:
 	/** The compressed morph target data, ready for the GPU. */
-	TSharedPtr<FExternalMorphTargetSet> MorphTargetSet;
+	TSharedPtr<FExternalMorphSet> MorphTargetSet;
 
 	/** The morph target set ID that is passed to FSkeletalMeshLODRenderData::AddExternalMorphBuffer(...). */
 	static int32 NeuralMorphsExternalMorphSetID;
diff --git a/Engine/Source/Runtime/Engine/Classes/Components/SkinnedMeshComponent.h b/Engine/Source/Runtime/Engine/Classes/Components/SkinnedMeshComponent.h
--- a/Engine/Source/Runtime/Engine/Classes/Components/SkinnedMeshComponent.h
+++ b/Engine/Source/Runtime/Engine/Classes/Components/SkinnedMeshComponent.h
@@ -15,6 +15,7 @@
 #include "Containers/SortedMap.h"
 #include "LODSyncInterface.h"
 #include "BoneContainer.h"
+#include "Rendering/MorphTargetVertexInfoBuffers.h"
 #include "SkinnedMeshComponent.generated.h"
 
 enum class ESkinCacheUsage : uint8;
@@ -140,6 +141,18 @@ struct FActiveMorphTarget
 	}
 };
 
+/** An external morph target set. External morph targets are managed by systems outside of the skinned meshes. */
+struct ENGINE_API FExternalMorphSet
+{
+	/** A name for this set, useful for debugging. */
+	FName Name = FName(TEXT("Unknown"));
+
+	/** The GPU compressed morph buffers. */
+	FMorphTargetVertexInfoBuffers MorphBuffers;
+};
+
+using FExternalMorphSets = TMap<int32, TSharedPtr<FExternalMorphSet>>;
+
 /** The weight data for a specific external morph set. */
 struct ENGINE_API FExternalMorphSetWeights
 {
@@ -420,9 +433,37 @@ public:
 	UE_DEPRECATED(5.1, "This method has been deprecated. Please use the GetLeaderBoneMap.")
 	const TArray<int32>& GetMasterBoneMap() const { return GetLeaderBoneMap(); }
 
+	/** Get the weights in read-only mode for a given external morph target set at a specific LOD. */
 	const FExternalMorphWeightData& GetExternalMorphWeights(int32 LOD) const { return ExternalMorphWeightData[LOD]; }
+
+	/** Get the weights for a given external morph target set at a specific LOD. */
 	FExternalMorphWeightData& GetExternalMorphWeights(int32 LOD) { return ExternalMorphWeightData[LOD]; }
 
+	/**
+	 * Register an external set of GPU compressed morph targets.
+	 * These compressed morph targets are GPU only morph targets that will not appear inside the UI and are owned by external systems.
+	 * Every set of these morph targets has some unique ID.
+	 * @param LOD The LOD index.
+	 * @param ID The unique ID for this set of morph targets.
+	 * @param MorphSet A shared pointer to a morph set, which basically contains a GPU friendly morph buffer. This buffer should be owned by an external system that calls this method.
+	 */
+	void AddExternalMorphSet(int32 LOD, int32 ID, TSharedPtr<FExternalMorphSet> MorphSet);
+
+	/** Remove a given set of external GPU based morph targets. */
+	void RemoveExternalMorphSet(int32 LOD, int32 ID);
+
+	/** Clear all externally registered morph target buffers. */
+	void ClearExternalMorphSets(int32 LOD);
+
+	/** Do we have a given set of external morph targets? */
+	bool HasExternalMorphSet(int32 LOD, int32 ID) const;
+
+	/** Get the external morph sets for a given LOD. */
+	const FExternalMorphSets& GetExternalMorphSets(int32 LOD) const { return ExternalMorphSets[LOD]; }
+
+	/** Get the array of external morph target sets. It is an array, one entry for each LOD. */
+	const TArray<FExternalMorphSets>& GetExternalMorphSetsArray() const { return ExternalMorphSets; }
+
 	/** 
 	 * Get CPU skinned vertices for the specified LOD level. Includes morph targets if they are enabled.
 	 * Note: This function is very SLOW as it needs to flush the render thread.
@@ -456,6 +497,14 @@ public:
 	/** The external morph target set weight data, for each LOD. This data is (re)initialized by RefreshExternalMorphTargetWeights(). */
 	TArray<FExternalMorphWeightData> ExternalMorphWeightData;
 
+	/**
+	 * External GPU based morph target buffers, for each LOD.
+	 * This contains an additional set of GPU only morphs that come from external other systems that generate morph targets.
+	 * These morph targets can only be updated on the GPU, will not be serialized as part of the Skeletal Mesh, and will not show up in the editors morph target list.
+	 * Every set of external morph targets has some given ID. Each LOD level has a map indexed by LOD number.
+	 */
+	TArray<FExternalMorphSets> ExternalMorphSets;
+
 #if WITH_EDITORONLY_DATA
 private:
 	/** Index of the section to preview... If set to -1, all section will be rendered */
@@ -1017,6 +1066,9 @@ public:
 	/** Get the pre-skinning local space bounds for this component. */
 	void GetPreSkinnedLocalBounds(FBoxSphereBounds& OutBounds) const;
 
+	/** Resize the morph target sets array to the number of LODs. */
+	void ResizeExternalMorphTargetSets();
+
 	/** Refresh the external morph target weight buffers. This makes sure the amount of morph sets and number of weights are valid. */
 	void RefreshExternalMorphTargetWeights();
 
diff --git a/Engine/Source/Runtime/Engine/Private/Components/SkinnedMeshComponent.cpp b/Engine/Source/Runtime/Engine/Private/Components/SkinnedMeshComponent.cpp
--- a/Engine/Source/Runtime/Engine/Private/Components/SkinnedMeshComponent.cpp
+++ b/Engine/Source/Runtime/Engine/Private/Components/SkinnedMeshComponent.cpp
@@ -580,27 +580,75 @@ void USkinnedMeshComponent::OnUnregister()
 	MeshDeformerInstance = nullptr;
 }
 
+void USkinnedMeshComponent::AddExternalMorphSet(int32 LOD, int32 ID, TSharedPtr<FExternalMorphSet> MorphSet)
+{
+	const int32 NumLODs = GetNumLODs();
+	if (NumLODs != ExternalMorphSets.Num())
+	{
+		ResizeExternalMorphTargetSets();
+	}
+	ExternalMorphSets[LOD].Add(ID, MorphSet);
+}
+
+void USkinnedMeshComponent::RemoveExternalMorphSet(int32 LOD, int32 ID)
+{
+	const int32 NumLODs = GetNumLODs();
+	if (NumLODs != ExternalMorphSets.Num())
+	{
+		ResizeExternalMorphTargetSets();
+	}
+	ExternalMorphSets[LOD].Remove(ID);
+}
+
+bool USkinnedMeshComponent::HasExternalMorphSet(int32 LOD, int32 ID) const
+{
+	check(LOD < ExternalMorphSets.Num());
+	return (ExternalMorphSets[LOD].Find(ID) != nullptr);
+}
+
+void USkinnedMeshComponent::ClearExternalMorphSets(int32 LOD)
+{
+	const int32 NumLODs = GetNumLODs();
+	if (NumLODs != ExternalMorphSets.Num())
+	{
+		ResizeExternalMorphTargetSets();
+	}
+	ExternalMorphSets[LOD].Empty();
+}
+
+void USkinnedMeshComponent::ResizeExternalMorphTargetSets()
+{
+	ExternalMorphSets.Reset();
+	ExternalMorphSets.AddDefaulted(GetNumLODs());
+}
+
 void USkinnedMeshComponent::RefreshExternalMorphTargetWeights()
 {
 	// Clear the external weights if there is no skeletal mesh.
-	if (GetSkinnedAsset() == nullptr || GetSkinnedAsset()->GetResourceForRendering() == nullptr)
+	if (GetSkinnedAsset() == nullptr)
 	{
 		ExternalMorphWeightData.Empty();
 		ExternalMorphWeightData.AddDefaulted(GetNumLODs());
 		return;
 	}
 
-	// Init based on the highest LOD.
 	const int32 NumLODs = GetNumLODs();
+	if (NumLODs != ExternalMorphSets.Num())
+	{
+		ResizeExternalMorphTargetSets();
+	}
+	check(ExternalMorphSets.Num() == NumLODs);
+
 	ExternalMorphWeightData.Reset();
 	ExternalMorphWeightData.AddDefaulted(NumLODs);
-	check(GetSkinnedAsset()->GetResourceForRendering()->LODRenderData.Num() == NumLODs);
 
+	// Make sure that for every LOD's list of morph target sets, we have the correct number of morph target weights.
+	// This resets all weights to 0 as well.
 	for (int32 LOD = 0; LOD < NumLODs; ++LOD)
 	{
 		FExternalMorphWeightData& ExternalWeights = GetExternalMorphWeights(LOD);
 
-		auto& MorphSets = GetSkinnedAsset()->GetResourceForRendering()->LODRenderData[LOD].ExternalMorphSets;
+		auto& MorphSets = ExternalMorphSets[LOD];
 		for (const auto& Item : MorphSets)
 		{
 			check(Item.Value.IsValid());
diff --git a/Engine/Source/Runtime/Engine/Private/SkeletalMeshLODRenderData.cpp b/Engine/Source/Runtime/Engine/Private/SkeletalMeshLODRenderData.cpp
--- a/Engine/Source/Runtime/Engine/Private/SkeletalMeshLODRenderData.cpp
+++ b/Engine/Source/Runtime/Engine/Private/SkeletalMeshLODRenderData.cpp
@@ -961,23 +961,3 @@ void FSkeletalMeshLODRenderData::GetSectionFromVertexIndex(int32 InVertIndex, in
 	// InVertIndex should always be in some chunk!
 	//check(false);
 }
-
-void FSkeletalMeshLODRenderData::AddExternalMorphSet(int32 ID, TSharedPtr<FExternalMorphTargetSet> MorphSet)
-{
-	ExternalMorphSets.Add(ID, MorphSet);
-}
-
-void FSkeletalMeshLODRenderData::RemoveExternalMorphSet(int32 ID)
-{
-	ExternalMorphSets.Remove(ID);
-}
-
-bool FSkeletalMeshLODRenderData::HasExternalMorphSet(int32 ID) const
-{
-	return (ExternalMorphSets.Find(ID) != nullptr);
-}
-
-void FSkeletalMeshLODRenderData::ClearExternalMorphSets()
-{
-	ExternalMorphSets.Empty();
-}
diff --git a/Engine/Source/Runtime/Engine/Private/SkeletalRenderGPUSkin.cpp b/Engine/Source/Runtime/Engine/Private/SkeletalRenderGPUSkin.cpp
--- a/Engine/Source/Runtime/Engine/Private/SkeletalRenderGPUSkin.cpp
+++ b/Engine/Source/Runtime/Engine/Private/SkeletalRenderGPUSkin.cpp
@@ -918,7 +918,7 @@ void FSkeletalMeshObjectGPUSkin::UpdateRayTracingGeometry(FSkeletalMeshLODRender
 #endif // RHI_RAYTRACING
 
 
-int32 FSkeletalMeshObjectGPUSkin::CalcNumActiveGPUMorphSets(FMorphVertexBuffer& MorphVertexBuffer, const FSkeletalMeshLODRenderData& LODData) const
+int32 FSkeletalMeshObjectGPUSkin::CalcNumActiveGPUMorphSets(FMorphVertexBuffer& MorphVertexBuffer, const FExternalMorphSets& ExternalMorphSets) const
 {
 	if (!UseGPUMorphTargets(FeatureLevel) || !IsValidRef(MorphVertexBuffer.VertexBufferRHI))
 	{
@@ -927,7 +927,7 @@ int32 FSkeletalMeshObjectGPUSkin::CalcNumActiveGPUMorphSets(FMorphVertexBuffer&
 
 	// Count all active external morph sets.
 	int32 NumMorphSets = 1; // Start at one, as we have our standard morph targets as well.
-	for (const auto& MorphSet : LODData.ExternalMorphSets)
+	for (const auto& MorphSet : ExternalMorphSets)
 	{
 		if (IsExternalMorphSetActive(MorphSet.Key, *MorphSet.Value))
 		{
@@ -937,7 +937,7 @@ int32 FSkeletalMeshObjectGPUSkin::CalcNumActiveGPUMorphSets(FMorphVertexBuffer&
 	return NumMorphSets;
 }
 
-bool FSkeletalMeshObjectGPUSkin::IsExternalMorphSetActive(int32 MorphSetID, FExternalMorphTargetSet& MorphSet) const
+bool FSkeletalMeshObjectGPUSkin::IsExternalMorphSetActive(int32 MorphSetID, const FExternalMorphSet& MorphSet) const
 {
 	const FMorphTargetVertexInfoBuffers& CompressedBuffers = MorphSet.MorphBuffers;
 	FExternalMorphSetWeights* WeightData = DynamicData->ExternalMorphWeightData.MorphSets.Find(MorphSetID);
@@ -978,7 +978,7 @@ static void CalculateMorphDeltaBoundsAccum(
 static void CalculateMorphDeltaBoundsIncludingExternalMorphs(
 	const TArray<float>& MorphTargetWeights,
 	const FMorphTargetVertexInfoBuffers& MorphTargetVertexInfoBuffers,
-	const TMap<int32, TSharedPtr<FExternalMorphTargetSet>>& ExternalMorphSets,
+	const FExternalMorphSets& ExternalMorphSets,
 	const TMap<int32, FExternalMorphSetWeights>& ExternalWeights,
 	FVector4& MorphScale,
 	FVector4& InvMorphScale)
@@ -1030,7 +1030,7 @@ void FSkeletalMeshObjectGPUSkin::UpdateMorphVertexBuffer(FRHICommandListImmediat
 	QUICK_SCOPE_CYCLE_COUNTER(STAT_FSkeletalMeshObjectGPUSkin_ProcessUpdatedDynamicData_UpdateMorphBuffer);
 	if (UseGPUMorphTargets(FeatureLevel))
 	{
-		const int32 NumMorphSets = CalcNumActiveGPUMorphSets(MorphVertexBuffer, LODData);
+		const int32 NumMorphSets = CalcNumActiveGPUMorphSets(MorphVertexBuffer, DynamicData->ExternalMorphSets);
 		int32 MorphSetIndex = 0;
 
 		// Calculate the delta bounds.
@@ -1041,7 +1041,7 @@ void FSkeletalMeshObjectGPUSkin::UpdateMorphVertexBuffer(FRHICommandListImmediat
 			CalculateMorphDeltaBoundsIncludingExternalMorphs(
 				DynamicData->MorphTargetWeights,
 				LODData.MorphTargetVertexInfoBuffers,
-				LODData.ExternalMorphSets,
+				DynamicData->ExternalMorphSets,
 				DynamicData->ExternalMorphWeightData.MorphSets,
 				MorphScale,
 				InvMorphScale);
@@ -1065,7 +1065,7 @@ void FSkeletalMeshObjectGPUSkin::UpdateMorphVertexBuffer(FRHICommandListImmediat
 		MorphSetIndex++;
 
 		// Process all external morph targets.
-		for (const auto& MorphSet : LODData.ExternalMorphSets)
+		for (const auto& MorphSet : DynamicData->ExternalMorphSets)
 		{
 			const int32 MorphSetID = MorphSet.Key;
 			const FMorphTargetVertexInfoBuffers& CompressedBuffers = MorphSet.Value->MorphBuffers;
@@ -2212,6 +2212,7 @@ void FDynamicSkelMeshObjectDataGPUSkin::Clear()
 	ActiveMorphTargets.Reset();
 	MorphTargetWeights.Reset();
 	ExternalMorphWeightData.Reset();
+	ExternalMorphSets.Reset();
 	NumWeightedActiveMorphTargets = 0;
 	ClothingSimData.Reset();
 	ClothBlendWeight = 0.0f;
@@ -2323,8 +2324,10 @@ void FDynamicSkelMeshObjectDataGPUSkin::InitDynamicSkelMeshObjectDataGPUSkin(
 	ActiveMorphTargets.Append(InActiveMorphTargets);
 	MorphTargetWeights.Append(InMorphTargetWeights);
 	NumWeightedActiveMorphTargets = 0;
-	ExternalMorphWeightData = InExternalMorphWeightData;	// TODO: Make this more efficient?
+
+	ExternalMorphWeightData = InExternalMorphWeightData;
 	ExternalMorphWeightData.UpdateNumActiveMorphTargets();
+	ExternalMorphSets = InMeshComponent->GetExternalMorphSets(InLODIndex);
 
 	// Gather any bones referenced by shadow shapes
 	FSkeletalMeshSceneProxy* SkeletalMeshProxy = InMeshComponent ? (FSkeletalMeshSceneProxy*)InMeshComponent->SceneProxy : nullptr;
diff --git a/Engine/Source/Runtime/Engine/Private/SkeletalRenderGPUSkin.h b/Engine/Source/Runtime/Engine/Private/SkeletalRenderGPUSkin.h
--- a/Engine/Source/Runtime/Engine/Private/SkeletalRenderGPUSkin.h
+++ b/Engine/Source/Runtime/Engine/Private/SkeletalRenderGPUSkin.h
@@ -101,6 +101,9 @@ public:
 	 */
 	FExternalMorphWeightData ExternalMorphWeightData;
 
+	/** The external morph target sets for this specific LOD. */
+	FExternalMorphSets ExternalMorphSets;
+
 	/** data for updating cloth section */
 	TMap<int32, FClothSimulData> ClothingSimData;
 
@@ -415,10 +418,10 @@ public:
 	 * Calculate how many GPU compressed morph target sets are active.
 	 * This includes regular morph targets as well as external morph targets.
 	 */
-	int32 CalcNumActiveGPUMorphSets(FMorphVertexBuffer& MorphVertexBuffer, const FSkeletalMeshLODRenderData& LODData) const;
+	int32 CalcNumActiveGPUMorphSets(FMorphVertexBuffer& MorphVertexBuffer, const FExternalMorphSets& ExternalMorphSets) const;
 
 	/** Check if a given morph set is active or not. If so, we will process it. */
-	bool IsExternalMorphSetActive(int32 MorphSetID, FExternalMorphTargetSet& MorphSet) const;
+	bool IsExternalMorphSetActive(int32 MorphSetID, const FExternalMorphSet& MorphSet) const;
 
 	FSkinWeightVertexBuffer* GetSkinWeightVertexBuffer(int32 LODIndex) const;
 
diff --git a/Engine/Source/Runtime/Engine/Public/Rendering/SkeletalMeshLODRenderData.h b/Engine/Source/Runtime/Engine/Public/Rendering/SkeletalMeshLODRenderData.h
--- a/Engine/Source/Runtime/Engine/Public/Rendering/SkeletalMeshLODRenderData.h
+++ b/Engine/Source/Runtime/Engine/Public/Rendering/SkeletalMeshLODRenderData.h
@@ -21,17 +21,6 @@ class USkinnedAsset;
 class FSkeletalMeshLODModel;
 #endif // WITH_EDITOR
 
-
-/** An external morph target set. External morph targets are managed by systems outside of the skinned meshes. */
-struct ENGINE_API FExternalMorphTargetSet
-{
-	/** A name for this set, useful for debugging. */
-	FName Name = FName(TEXT("Unknown"));
-
-	/** The GPU compressed morph buffers. */
-	FMorphTargetVertexInfoBuffers MorphBuffers;
-};
-
 struct FSkelMeshRenderSection
 {
 	/** Material (texture) used for this section. */
@@ -149,15 +138,6 @@ public:
 	/** GPU friendly access data for MorphTargets for an LOD */
 	FMorphTargetVertexInfoBuffers	MorphTargetVertexInfoBuffers;
 
-	/** 
-	 * External GPU based morph target buffers.
-	 * This contains an additional set of GPU only morphs that come from external other systems that generate morph targets.
-	 * These morph targets can only be updated on the GPU, will not be serialized as part of the Skeletal Mesh, and will not show up in the editors morph target list.
-	 * We only register them inside the render data so that these morph targets can get applied when rendering.
-	 * Every set of external morph targets has some given ID.
-	 */
-	TMap<int32, TSharedPtr<FExternalMorphTargetSet>> ExternalMorphSets;
-
 	/** Skin weight profile data structures, can contain multiple profiles and their runtime FSkinWeightVertexBuffer */
 	FSkinWeightProfilesData SkinWeightProfilesData;
 
@@ -286,24 +266,6 @@ public:
 		return OverrideBuffer != nullptr ? OverrideBuffer : &SkinWeightVertexBuffer;
 	}
 
-	/** 
-	 * Register an external set of GPU compressed morph targets. 
-	 * These compressed morph targets are GPU only morph targets that will not appear inside the UI and are owned by external systems.
-	 * Every set of these morph targets has some unique ID.
-	 * @param ID The unique ID for this set of morph targets.
-	 * @param MorphSet A shared pointer to a morph set, which basically contains a GPU friendly morph buffer. This buffer should be owned by an external system that calls this method.
-	 */
-	ENGINE_API void AddExternalMorphSet(int32 ID, TSharedPtr<FExternalMorphTargetSet> MorphSet);
-
-	/** Remove a given set of external GPU based morph targets. */
-	ENGINE_API void RemoveExternalMorphSet(int32 ID);
-
-	/** Clear all externally registered morph target buffers. */
-	ENGINE_API void ClearExternalMorphSets();
-
-	/** Do we have a given set of external morph targets? */
-	ENGINE_API bool HasExternalMorphSet(int32 ID) const;
-
 	/** Utility function for returning total number of faces in this LOD. */
 	ENGINE_API int32 GetTotalFaces() const;
 
