[MLDeformer]
- Vertex delta model got merged into the Neural Morph Model, separating them by a new mode property: Local and Global (previously Vertex Delta Model).
- Local and Global have a few slight different settings. The UI will react to it, and defaults are different.
- The morph target debug rendering 'Morph Target Number' property now is clamped to the maximum number of morph targets in the trained model.
- Exclude training and ground truth testing assets from the package build automatically.

diff --git a/Engine/Plugins/Experimental/MLDeformer/MLDeformerFramework/Source/MLDeformerFramework/Private/MLDeformerModel.cpp b/Engine/Plugins/Experimental/MLDeformer/MLDeformerFramework/Source/MLDeformerFramework/Private/MLDeformerModel.cpp
--- a/Engine/Plugins/Experimental/MLDeformer/MLDeformerFramework/Source/MLDeformerFramework/Private/MLDeformerModel.cpp
+++ b/Engine/Plugins/Experimental/MLDeformer/MLDeformerFramework/Source/MLDeformerFramework/Private/MLDeformerModel.cpp
@@ -13,6 +13,7 @@
 #include "GeometryCacheMeshData.h"
 #include "GeometryCacheTrack.h"
 #include "NeuralNetwork.h"
+#include "Animation/AnimSequence.h"
 #include "UObject/UObjectGlobals.h"
 
 namespace UE::MLDeformer
@@ -117,6 +118,10 @@ void UMLDeformerModel::PostLoad()
 			NeuralNetwork->SetDeviceType(ENeuralDeviceType::CPU, ENeuralDeviceType::CPU, ENeuralDeviceType::CPU);
 		}
 	}
+
+#if WITH_EDITOR
+	SetAssetEditorOnlyFlags();
+#endif
 }
 
 void UMLDeformerModel::SetNeuralNetwork(UNeuralNetwork* InNeuralNetwork)
@@ -199,6 +204,14 @@ void UMLDeformerModel::FloatArrayToVector3Array(const TArray<float>& FloatArray,
 	{
 		return SkeletalMesh ? SkeletalMesh->GetNumImportedVertices() : 0;
 	}
+
+	void UMLDeformerModel::SetAssetEditorOnlyFlags()
+	{
+		if (AnimSequence)
+		{
+			AnimSequence->GetPackage()->SetPackageFlags(PKG_EditorOnly);
+		}
+	}
 #endif	// #if WITH_EDITOR
 
 #if WITH_EDITORONLY_DATA
diff --git a/Engine/Plugins/Experimental/MLDeformer/MLDeformerFramework/Source/MLDeformerFramework/Public/MLDeformerModel.h b/Engine/Plugins/Experimental/MLDeformer/MLDeformerFramework/Source/MLDeformerFramework/Public/MLDeformerModel.h
--- a/Engine/Plugins/Experimental/MLDeformer/MLDeformerFramework/Source/MLDeformerFramework/Public/MLDeformerModel.h
+++ b/Engine/Plugins/Experimental/MLDeformer/MLDeformerFramework/Source/MLDeformerFramework/Public/MLDeformerModel.h
@@ -69,6 +69,7 @@ public:
 	virtual void UpdateNumBaseMeshVertices();
 	virtual void UpdateCachedNumVertices();
 	virtual void UpdateNumTargetMeshVertices() PURE_VIRTUAL(UMLDeformerModel::GetNumTargetMeshVertices; mTargetMeshVerts = 0;);
+	virtual void SetAssetEditorOnlyFlags();
 	static int32 ExtractNumImportedSkinnedVertices(const USkeletalMesh* SkeletalMesh);
 #endif
 
diff --git a/Engine/Plugins/Experimental/MLDeformer/MLDeformerFramework/Source/MLDeformerFrameworkEditor/Private/MLDeformerEditorModel.cpp b/Engine/Plugins/Experimental/MLDeformer/MLDeformerFramework/Source/MLDeformerFrameworkEditor/Private/MLDeformerEditorModel.cpp
--- a/Engine/Plugins/Experimental/MLDeformer/MLDeformerFramework/Source/MLDeformerFrameworkEditor/Private/MLDeformerEditorModel.cpp
+++ b/Engine/Plugins/Experimental/MLDeformer/MLDeformerFramework/Source/MLDeformerFrameworkEditor/Private/MLDeformerEditorModel.cpp
@@ -1547,6 +1547,7 @@ namespace UE::MLDeformer
 	{
 		OnInputAssetsChanged();
 		OnPostInputAssetChanged();
+		Model->SetAssetEditorOnlyFlags();
 		GetEditor()->GetModelDetailsView()->ForceRefresh();
 		if (RefreshVizSettings)
 		{
diff --git a/Engine/Plugins/Experimental/MLDeformer/NeuralMorphModel/Source/NeuralMorphModel/Private/NeuralMorphModel.cpp b/Engine/Plugins/Experimental/MLDeformer/NeuralMorphModel/Source/NeuralMorphModel/Private/NeuralMorphModel.cpp
--- a/Engine/Plugins/Experimental/MLDeformer/NeuralMorphModel/Source/NeuralMorphModel/Private/NeuralMorphModel.cpp
+++ b/Engine/Plugins/Experimental/MLDeformer/NeuralMorphModel/Source/NeuralMorphModel/Private/NeuralMorphModel.cpp
@@ -60,6 +60,28 @@ UMLDeformerModelInstance* UNeuralMorphModel::CreateModelInstance(UMLDeformerComp
 	{
 		NumTargetMeshVerts = UE::MLDeformer::ExtractNumImportedGeomCacheVertices(GeometryCache);
 	}
+
+	void UNeuralMorphModel::SetAssetEditorOnlyFlags()
+	{
+		// Set the flags for the base class, which filters out the training anim sequence.
+		UMLDeformerModel::SetAssetEditorOnlyFlags();
+
+		// The training geometry cache is something we don't want to package.
+		if (GeometryCache)
+		{
+			GeometryCache->GetPackage()->SetPackageFlags(PKG_EditorOnly);
+		}
+
+		// Filter the viz settings specific assets.
+		UNeuralMorphModelVizSettings* NeuralVizSettings = Cast<UNeuralMorphModelVizSettings>(VizSettings);
+		if (NeuralVizSettings)
+		{
+			if (NeuralVizSettings->GetTestGroundTruth())
+			{
+				NeuralVizSettings->GetTestGroundTruth()->GetPackage()->SetPackageFlags(PKG_EditorOnly);
+			}
+		}
+	}
 #endif // WITH_EDITOR
 
 #if WITH_EDITORONLY_DATA
diff --git a/Engine/Plugins/Experimental/MLDeformer/NeuralMorphModel/Source/NeuralMorphModel/Public/NeuralMorphModel.h b/Engine/Plugins/Experimental/MLDeformer/NeuralMorphModel/Source/NeuralMorphModel/Public/NeuralMorphModel.h
--- a/Engine/Plugins/Experimental/MLDeformer/NeuralMorphModel/Source/NeuralMorphModel/Public/NeuralMorphModel.h
+++ b/Engine/Plugins/Experimental/MLDeformer/NeuralMorphModel/Source/NeuralMorphModel/Public/NeuralMorphModel.h
@@ -19,6 +19,23 @@ class USkeleton;
 class IPropertyHandle;
 struct FExternalMorphSet;
 
+UENUM()
+enum class ENeuralMorphMode : uint8
+{
+	/**
+	 * Each bone creates a set of morph targets and has its own small neural network.
+	 * This mode is faster to process on the CPU side.
+	 */
+	Local,
+
+	/** 
+	 * There is one fully connected neural network that generates a set of morph targets.
+	 * This has a slightly higher CPU overhead, but could result in higher quality.
+	 */
+	Global
+};
+
+
 UCLASS()
 class NEURALMORPHMODEL_API UNeuralMorphModel 
 	: public UMLDeformerModel
@@ -39,7 +56,8 @@ public:
 	virtual void SampleGroundTruthPositions(float SampleTime, TArray<FVector3f>& OutPositions) override;
 #endif
 #if WITH_EDITOR
-	virtual void UpdateNumTargetMeshVertices();
+	virtual void UpdateNumTargetMeshVertices() override;
+	virtual void SetAssetEditorOnlyFlags() override;
 #endif
 	// ~END UMLDeformerModel overrides.
 
@@ -50,8 +68,10 @@ public:
 #if WITH_EDITORONLY_DATA
 	const UGeometryCache* GetGeometryCache() const { return GeometryCache; }
 	UGeometryCache* GetGeometryCache() { return GeometryCache; }
-	int32 GetNumHiddenLayers() const { return NumHiddenLayers; }
-	int32 GetNumNeuronsPerLayer() const { return NumNeuronsPerLayer; }
+	int32 GetLocalNumHiddenLayers() const { return LocalNumHiddenLayers; }
+	int32 GetLocalNumNeuronsPerLayer() const { return LocalNumNeuronsPerLayer; }
+	int32 GetGlobalNumHiddenLayers() const { return GlobalNumHiddenLayers; }
+	int32 GetGlobalNumNeuronsPerLayer() const { return GlobalNumNeuronsPerLayer; }
 	int32 GetNumIterations() const { return NumIterations; }
 	int32 GetBatchSize() const { return BatchSize; }
 	float GetLearningRate() const { return LearningRate; }
@@ -88,20 +108,36 @@ public:
 	TObjectPtr<UGeometryCache> GeometryCache = nullptr;
 
 	/** The number of morph targets to generate per bone. */
-	UPROPERTY(EditAnywhere, BlueprintReadWrite, Category = "Training Settings", meta = (ClampMin = "1", ClampMax = "100"))
-	int32 NumMorphTargetsPerBone = 6;
+	UPROPERTY(EditAnywhere, BlueprintReadWrite, Category = "Training Settings")
+	ENeuralMorphMode Mode = ENeuralMorphMode::Local;
+
+	/** The number of morph targets to generate per bone. */
+	UPROPERTY(EditAnywhere, BlueprintReadWrite, DisplayName = "Num Morph Targets Per Bone", Category = "Training Settings", meta = (ClampMin = "1", ClampMax = "1000"))
+	int32 LocalNumMorphTargetsPerBone = 6;
+
+	/** The number of morph targets to generate. */
+	UPROPERTY(EditAnywhere, BlueprintReadWrite, DisplayName = "Num Morph Targets", Category = "Training Settings", meta = (ClampMin = "1", ClampMax = "1000"))
+	int32 GlobalNumMorphTargets = 128;
 
 	/** The number of iterations to train the model for. */
-	UPROPERTY(EditAnywhere, BlueprintReadWrite, Category = "Training Settings", meta = (ClampMin = "1"))
+	UPROPERTY(EditAnywhere, BlueprintReadWrite, Category = "Training Settings", meta = (ClampMin = "1", ClampMax = "500000"))
 	int32 NumIterations = 2000;
 
 	/** The number of hidden layers that the neural network model will have.\nHigher numbers will slow down performance but can deal with more complex deformations. */
-	UPROPERTY(EditAnywhere, BlueprintReadWrite, AdvancedDisplay, Category = "Training Settings", meta = (ClampMin = "1", ClampMax = "10"))
-	int32 NumHiddenLayers = 1;
+	UPROPERTY(EditAnywhere, BlueprintReadWrite, AdvancedDisplay, DisplayName = "Num Hidden Layers", Category = "Training Settings", meta = (ClampMin = "1", ClampMax = "5"))
+	int32 LocalNumHiddenLayers = 1;
 
 	/** The number of units/neurons per hidden layer. Higher numbers will slow down performance but allow for more complex mesh deformations. */
-	UPROPERTY(EditAnywhere, BlueprintReadWrite, AdvancedDisplay, Category = "Training Settings", meta = (ClampMin = "1"))
-	int32 NumNeuronsPerLayer = 6;
+	UPROPERTY(EditAnywhere, BlueprintReadWrite, AdvancedDisplay, DisplayName = "Num Neurons Per Layer", Category = "Training Settings", meta = (ClampMin = "1"))
+	int32 LocalNumNeuronsPerLayer = 6;
+
+	/** The number of hidden layers that the neural network model will have.\nHigher numbers will slow down performance but can deal with more complex deformations. */
+	UPROPERTY(EditAnywhere, BlueprintReadWrite, AdvancedDisplay, DisplayName = "Num Hidden Layers", Category = "Training Settings", meta = (ClampMin = "1", ClampMax = "5"))
+	int32 GlobalNumHiddenLayers = 2;
+
+	/** The number of units/neurons per hidden layer. Higher numbers will slow down performance but allow for more complex mesh deformations. */
+	UPROPERTY(EditAnywhere, BlueprintReadWrite, AdvancedDisplay, DisplayName = "Num Neurons Per Layer", Category = "Training Settings", meta = (ClampMin = "1"))
+	int32 GlobalNumNeuronsPerLayer = 128;
 
 	/** The number of frames per batch when training the model. */
 	UPROPERTY(EditAnywhere, BlueprintReadWrite, AdvancedDisplay, Category = "Training Settings", meta = (ClampMin = "1"))
@@ -127,7 +163,7 @@ public:
 	float MorphTargetDeltaThreshold = 0.0025f;
 
 	/** The morph target error tolerance. Higher values result in larger compression, but could result in visual artifacts. */
-	UPROPERTY(EditAnywhere, BlueprintReadWrite, AdvancedDisplay, Category = "Training Settings", meta = (ClampMin = "1", ClampMax = "1000"))
+	UPROPERTY(EditAnywhere, BlueprintReadWrite, AdvancedDisplay, Category = "Training Settings", meta = (ClampMin = "0.01", ClampMax = "1000"))
 	float MorphTargetErrorTolerance = 50.0f;
 #endif // WITH_EDITORONLY_DATA
 };
diff --git a/Engine/Plugins/Experimental/MLDeformer/NeuralMorphModel/Source/NeuralMorphModelEditor/Private/NeuralMorphEditorModel.cpp b/Engine/Plugins/Experimental/MLDeformer/NeuralMorphModel/Source/NeuralMorphModelEditor/Private/NeuralMorphEditorModel.cpp
--- a/Engine/Plugins/Experimental/MLDeformer/NeuralMorphModel/Source/NeuralMorphModelEditor/Private/NeuralMorphEditorModel.cpp
+++ b/Engine/Plugins/Experimental/MLDeformer/NeuralMorphModel/Source/NeuralMorphModelEditor/Private/NeuralMorphEditorModel.cpp
@@ -73,6 +73,21 @@ namespace UE::NeuralMorphModel
 				SampleDeltas();
 			}
 		}
+		else if (Property->GetFName() == GET_MEMBER_NAME_CHECKED(UNeuralMorphModel, Mode))
+		{
+			if (PropertyChangedEvent.ChangeType == EPropertyChangeType::ValueSet)
+			{
+				SetResamplingInputOutputsNeeded(true);	// Needed as local mode doesn't support curves, while the global mode does.
+				GetEditor()->GetModelDetailsView()->ForceRefresh();
+			}
+		}
+		else if (Property->GetFName() == GET_MEMBER_NAME_CHECKED(UNeuralMorphModelVizSettings, MorphTargetNumber))
+		{
+			UNeuralMorphModel* NeuralModel = GetNeuralMorphModel();
+			UNeuralMorphModelVizSettings* NeuralViz = GetNeuralMorphModelVizSettings();
+			const int32 NumMorphTargets = NeuralModel->MorphTargetSet.IsValid() ? NeuralModel->MorphTargetSet->MorphBuffers.GetNumMorphs() : 0;
+			NeuralViz->MorphTargetNumber = (NumMorphTargets > 0) ? FMath::Min<int32>(NeuralViz->MorphTargetNumber, NumMorphTargets - 1) : 0;
+		}
 	}
 
 	void FNeuralMorphEditorModel::OnInputAssetsChanged()
diff --git a/Engine/Plugins/Experimental/MLDeformer/NeuralMorphModel/Source/NeuralMorphModelEditor/Private/NeuralMorphModelDetails.cpp b/Engine/Plugins/Experimental/MLDeformer/NeuralMorphModel/Source/NeuralMorphModelEditor/Private/NeuralMorphModelDetails.cpp
--- a/Engine/Plugins/Experimental/MLDeformer/NeuralMorphModel/Source/NeuralMorphModelEditor/Private/NeuralMorphModelDetails.cpp
+++ b/Engine/Plugins/Experimental/MLDeformer/NeuralMorphModel/Source/NeuralMorphModelEditor/Private/NeuralMorphModelDetails.cpp
@@ -59,9 +59,22 @@ namespace UE::NeuralMorphModel
 		FMLDeformerModelDetails::CustomizeDetails(DetailBuilder);
 
 		// Training settings.
-		SettingsCategoryBuilder->AddProperty(GET_MEMBER_NAME_CHECKED(UNeuralMorphModel, NumMorphTargetsPerBone));
-		SettingsCategoryBuilder->AddProperty(GET_MEMBER_NAME_CHECKED(UNeuralMorphModel, NumHiddenLayers));
-		SettingsCategoryBuilder->AddProperty(GET_MEMBER_NAME_CHECKED(UNeuralMorphModel, NumNeuronsPerLayer));
+		SettingsCategoryBuilder->AddProperty(GET_MEMBER_NAME_CHECKED(UNeuralMorphModel, Mode));
+
+		SettingsCategoryBuilder->AddProperty(GET_MEMBER_NAME_CHECKED(UNeuralMorphModel, LocalNumMorphTargetsPerBone))
+			.Visibility(NeuralMorphModel->Mode == ENeuralMorphMode::Local ? EVisibility::Visible : EVisibility::Collapsed);
+		SettingsCategoryBuilder->AddProperty(GET_MEMBER_NAME_CHECKED(UNeuralMorphModel, LocalNumHiddenLayers))
+			.Visibility(NeuralMorphModel->Mode == ENeuralMorphMode::Local ? EVisibility::Visible : EVisibility::Collapsed);
+		SettingsCategoryBuilder->AddProperty(GET_MEMBER_NAME_CHECKED(UNeuralMorphModel, LocalNumNeuronsPerLayer))
+			.Visibility(NeuralMorphModel->Mode == ENeuralMorphMode::Local ? EVisibility::Visible : EVisibility::Collapsed);
+
+		SettingsCategoryBuilder->AddProperty(GET_MEMBER_NAME_CHECKED(UNeuralMorphModel, GlobalNumMorphTargets))
+			.Visibility(NeuralMorphModel->Mode == ENeuralMorphMode::Global ? EVisibility::Visible : EVisibility::Collapsed);
+		SettingsCategoryBuilder->AddProperty(GET_MEMBER_NAME_CHECKED(UNeuralMorphModel, GlobalNumHiddenLayers))
+			.Visibility(NeuralMorphModel->Mode == ENeuralMorphMode::Global ? EVisibility::Visible : EVisibility::Collapsed);
+		SettingsCategoryBuilder->AddProperty(GET_MEMBER_NAME_CHECKED(UNeuralMorphModel, GlobalNumNeuronsPerLayer))
+			.Visibility(NeuralMorphModel->Mode == ENeuralMorphMode::Global ? EVisibility::Visible : EVisibility::Collapsed);
+
 		SettingsCategoryBuilder->AddProperty(GET_MEMBER_NAME_CHECKED(UNeuralMorphModel, NumIterations));
 		SettingsCategoryBuilder->AddProperty(GET_MEMBER_NAME_CHECKED(UNeuralMorphModel, BatchSize));
 		SettingsCategoryBuilder->AddProperty(GET_MEMBER_NAME_CHECKED(UNeuralMorphModel, LearningRate));
diff --git a/Engine/Plugins/Experimental/MLDeformer/VertexDeltaModel/Source/VertexDeltaModel/Private/VertexDeltaModel.cpp b/Engine/Plugins/Experimental/MLDeformer/VertexDeltaModel/Source/VertexDeltaModel/Private/VertexDeltaModel.cpp
--- a/Engine/Plugins/Experimental/MLDeformer/VertexDeltaModel/Source/VertexDeltaModel/Private/VertexDeltaModel.cpp
+++ b/Engine/Plugins/Experimental/MLDeformer/VertexDeltaModel/Source/VertexDeltaModel/Private/VertexDeltaModel.cpp
@@ -3,6 +3,7 @@
 #include "VertexDeltaModel.h"
 #include "VertexDeltaModelVizSettings.h"
 #include "MLDeformerGeomCacheHelpers.h"
+#include "GeometryCache.h"
 #include "UObject/Object.h"
 
 #define LOCTEXT_NAMESPACE "UVertexDeltaModel"
@@ -20,6 +21,28 @@ UVertexDeltaModel::UVertexDeltaModel(const FObjectInitializer& ObjectInitializer
 	{
 		NumTargetMeshVerts = UE::MLDeformer::ExtractNumImportedGeomCacheVertices(GeometryCache);
 	}
+
+	void UVertexDeltaModel::SetAssetEditorOnlyFlags()
+	{
+		// Set the flags for the base class, which filters out the training anim sequence.
+		UMLDeformerModel::SetAssetEditorOnlyFlags();
+
+		// The training geometry cache is something we don't want to package.
+		if (GeometryCache)
+		{
+			GeometryCache->GetPackage()->SetPackageFlags(PKG_EditorOnly);
+		}
+
+		// Filter the viz settings specific assets.
+		UVertexDeltaModelVizSettings* VertexDeltaVizSettings = Cast<UVertexDeltaModelVizSettings>(VizSettings);
+		if (VertexDeltaVizSettings)
+		{
+			if (VertexDeltaVizSettings->GetTestGroundTruth())
+			{
+				VertexDeltaVizSettings->GetTestGroundTruth()->GetPackage()->SetPackageFlags(PKG_EditorOnly);
+			}
+		}
+	}
 #endif // WITH_EDITOR
 
 #if WITH_EDITORONLY_DATA
diff --git a/Engine/Plugins/Experimental/MLDeformer/VertexDeltaModel/Source/VertexDeltaModel/Public/VertexDeltaModel.h b/Engine/Plugins/Experimental/MLDeformer/VertexDeltaModel/Source/VertexDeltaModel/Public/VertexDeltaModel.h
--- a/Engine/Plugins/Experimental/MLDeformer/VertexDeltaModel/Source/VertexDeltaModel/Public/VertexDeltaModel.h
+++ b/Engine/Plugins/Experimental/MLDeformer/VertexDeltaModel/Source/VertexDeltaModel/Public/VertexDeltaModel.h
@@ -34,7 +34,8 @@ public:
 	virtual void SampleGroundTruthPositions(float SampleTime, TArray<FVector3f>& OutPositions) override;
 #endif
 #if WITH_EDITOR
-	virtual void UpdateNumTargetMeshVertices();
+	virtual void UpdateNumTargetMeshVertices() override;
+	virtual void SetAssetEditorOnlyFlags() override;
 #endif
 	// ~END UMLDeformerModel overrides.
 
