[MLDeformer]
- Removed an unneeded operation in the local neural morph model python code.
- Don't render morph targets preview when in training mode.
- Make the timeline's play button do nothing in training mode, as that wouldn't update deltas. Also playing an anim with random frames wouldn't make sense.
- Some refactor, introducing a shared geometry cache actor, so neural morph model and vertex delta models don't have to create their own, as they were identical.

diff --git a/Engine/Plugins/Experimental/MLDeformer/MLDeformerFramework/Source/MLDeformerFrameworkEditor/Private/MLDeformerApplicationMode.cpp b/Engine/Plugins/Experimental/MLDeformer/MLDeformerFramework/Source/MLDeformerFrameworkEditor/Private/MLDeformerApplicationMode.cpp
--- a/Engine/Plugins/Experimental/MLDeformer/MLDeformerFramework/Source/MLDeformerFrameworkEditor/Private/MLDeformerApplicationMode.cpp
+++ b/Engine/Plugins/Experimental/MLDeformer/MLDeformerFramework/Source/MLDeformerFrameworkEditor/Private/MLDeformerApplicationMode.cpp
@@ -42,7 +42,7 @@ namespace UE::MLDeformer
 		TabFactories.RegisterFactory(MakeShared<FMLDeformerTimelineTabSummoner>(MLDeformerEditor));
 
 		// Create tab layout.
-		TabLayout = FTabManager::NewLayout("Standalone_MLDeformerEditor_Layout_v0.031")
+		TabLayout = FTabManager::NewLayout("Standalone_MLDeformerEditor_Layout_v0.032")
 			->AddArea
 			(
 				FTabManager::NewPrimaryArea()
@@ -73,7 +73,7 @@ namespace UE::MLDeformer
 						->Split
 						(
 							FTabManager::NewStack()
-							->SetSizeCoefficient(0.01f)
+							->SetSizeCoefficient(0.03f)
 							->SetHideTabWell(true)
 							->AddTab(FMLDeformerTimelineTabSummoner::TabID, ETabState::OpenedTab)
 						)
diff --git a/Engine/Plugins/Experimental/MLDeformer/MLDeformerFramework/Source/MLDeformerFrameworkEditor/Private/MLDeformerEditorActor.cpp b/Engine/Plugins/Experimental/MLDeformer/MLDeformerFramework/Source/MLDeformerFrameworkEditor/Private/MLDeformerEditorActor.cpp
--- a/Engine/Plugins/Experimental/MLDeformer/MLDeformerFramework/Source/MLDeformerFrameworkEditor/Private/MLDeformerEditorActor.cpp
+++ b/Engine/Plugins/Experimental/MLDeformer/MLDeformerFramework/Source/MLDeformerFrameworkEditor/Private/MLDeformerEditorActor.cpp
@@ -124,6 +124,16 @@ namespace UE::MLDeformer
 		}
 	}
 
+	bool FMLDeformerEditorActor::IsPlaying() const
+	{
+		if (SkeletalMeshComponent)
+		{
+			return !SkeletalMeshComponent->bPauseAnims;
+		}
+
+		return false;
+	}
+
 	FBox FMLDeformerEditorActor::GetBoundingBox() const
 	{
 		if (SkeletalMeshComponent)
diff --git a/Engine/Plugins/Experimental/MLDeformer/MLDeformerFramework/Source/MLDeformerFrameworkEditor/Private/MLDeformerEditorModel.cpp b/Engine/Plugins/Experimental/MLDeformer/MLDeformerFramework/Source/MLDeformerFrameworkEditor/Private/MLDeformerEditorModel.cpp
--- a/Engine/Plugins/Experimental/MLDeformer/MLDeformerFramework/Source/MLDeformerFrameworkEditor/Private/MLDeformerEditorModel.cpp
+++ b/Engine/Plugins/Experimental/MLDeformer/MLDeformerFramework/Source/MLDeformerFrameworkEditor/Private/MLDeformerEditorModel.cpp
@@ -689,17 +689,21 @@ namespace UE::MLDeformer
 
 	void FMLDeformerEditorModel::OnPlayPressed()
 	{
+		// The play button shouldn't do anything in training mode.
 		UMLDeformerVizSettings* VizSettings = Model->GetVizSettings();
-		const bool bTestData = VizSettings->GetVisualizationMode() == EMLDeformerVizMode::TestData;
-		FMLDeformerEditorActor* BaseActor = GetVisualizationModeBaseActor();
+		if (VizSettings->GetVisualizationMode() == EMLDeformerVizMode::TrainingData)
+		{
+			return;
+		}
 
+		FMLDeformerEditorActor* BaseActor = GetVisualizationModeBaseActor();
 		const bool bMustPause = (BaseActor && BaseActor->GetSkeletalMeshComponent()) ? !BaseActor->GetSkeletalMeshComponent()->bPauseAnims : false;
 
+		const bool bTestData = (VizSettings->GetVisualizationMode() == EMLDeformerVizMode::TestData);
 		for (FMLDeformerEditorActor* EditorActor : EditorActors)
 		{
 			if (EditorActor && ((EditorActor->IsTestActor() && bTestData) || (EditorActor->IsTrainingActor() && !bTestData)))
 			{
-
 				EditorActor->SetPlaySpeed(VizSettings->AnimPlaySpeed);
 				EditorActor->Pause(bMustPause);
 			}
diff --git a/Engine/Plugins/Experimental/MLDeformer/MLDeformerFramework/Source/MLDeformerFrameworkEditor/Private/MLDeformerGeomCacheActor.cpp b/Engine/Plugins/Experimental/MLDeformer/MLDeformerFramework/Source/MLDeformerFrameworkEditor/Private/MLDeformerGeomCacheActor.cpp
new file mode 100644
--- /dev/null
+++ b/Engine/Plugins/Experimental/MLDeformer/MLDeformerFramework/Source/MLDeformerFrameworkEditor/Private/MLDeformerGeomCacheActor.cpp
@@ -0,0 +1,156 @@
+// Copyright Epic Games, Inc. All Rights Reserved.
+
+#include "MLDeformerGeomCacheActor.h"
+#include "GeometryCacheComponent.h"
+#include "Animation/DebugSkelMeshComponent.h"
+
+namespace UE::MLDeformer
+{
+	FMLDeformerGeomCacheActor::FMLDeformerGeomCacheActor(const FConstructSettings& Settings)
+		: FMLDeformerEditorActor(Settings)
+	{
+	}
+
+	FMLDeformerGeomCacheActor::~FMLDeformerGeomCacheActor()
+	{
+		if (GeomCacheComponent)
+		{
+			Actor->RemoveOwnedComponent(GeomCacheComponent);
+		}
+	}
+
+	void FMLDeformerGeomCacheActor::SetVisibility(bool bIsVisible)
+	{
+		FMLDeformerEditorActor::SetVisibility(bIsVisible);
+
+		if (GeomCacheComponent && bIsVisible != GeomCacheComponent->IsVisible())
+		{
+			GeomCacheComponent->SetVisibility(bIsVisible, true);
+		}
+	}
+
+	bool FMLDeformerGeomCacheActor::IsVisible() const
+	{
+		if (SkeletalMeshComponent)
+		{
+			return SkeletalMeshComponent->IsVisible();
+		}
+
+		if (GeomCacheComponent)
+		{
+			return GeomCacheComponent->IsVisible();
+		}
+
+		return true;
+	}
+
+	bool FMLDeformerGeomCacheActor::HasVisualMesh() const
+	{
+		if (SkeletalMeshComponent && SkeletalMeshComponent->GetSkeletalMeshAsset())
+		{
+			return true;
+		}
+
+		if (GeomCacheComponent && GeomCacheComponent->GetGeometryCache())
+		{
+			return true;
+		}
+
+		return false;
+	}
+
+	void FMLDeformerGeomCacheActor::SetPlayPosition(float TimeInSeconds, bool bAutoPause)
+	{
+		if (SkeletalMeshComponent)
+		{
+			SkeletalMeshComponent->SetPosition(TimeInSeconds);
+			if (bAutoPause)
+			{
+				SkeletalMeshComponent->bPauseAnims = true;
+			}
+		}
+
+		if (GeomCacheComponent)
+		{
+			GeomCacheComponent->SetManualTick(true);
+			GeomCacheComponent->TickAtThisTime(TimeInSeconds, false, false, false);
+		}
+	}
+
+	float FMLDeformerGeomCacheActor::GetPlayPosition() const
+	{
+		if (SkeletalMeshComponent)
+		{
+			return SkeletalMeshComponent->GetPosition();
+		}
+
+		if (GeomCacheComponent)
+		{
+			float Duration = GeomCacheComponent->GetDuration();
+			float AnimTime = GeomCacheComponent->GetAnimationTime();
+			float StartTime = GeomCacheComponent->GetStartTimeOffset();
+			float DeltaTime = AnimTime - StartTime;
+			return DeltaTime > Duration ? FMath::Fmod(DeltaTime, Duration) : DeltaTime;
+		}
+
+		return 0.0f;
+	}
+
+	void FMLDeformerGeomCacheActor::SetPlaySpeed(float PlaySpeed)
+	{
+		if (SkeletalMeshComponent)
+		{
+			SkeletalMeshComponent->SetPlayRate(PlaySpeed);
+		}
+
+		if (GeomCacheComponent)
+		{
+			GeomCacheComponent->SetPlaybackSpeed(PlaySpeed);
+		}
+	}
+
+	void FMLDeformerGeomCacheActor::Pause(bool bPaused)
+	{
+		if (SkeletalMeshComponent)
+		{
+			SkeletalMeshComponent->bPauseAnims = bPaused;
+		}
+
+		if (GeomCacheComponent)
+		{
+			GeomCacheComponent->SetManualTick(bPaused);
+		}
+	}
+
+	bool FMLDeformerGeomCacheActor::IsPlaying() const
+	{
+		if (SkeletalMeshComponent)
+		{
+			return !SkeletalMeshComponent->bPauseAnims;
+		}
+
+		if (GeomCacheComponent)
+		{
+			return GeomCacheComponent->IsPlaying();
+		}
+
+		return false;
+	}
+
+	FBox FMLDeformerGeomCacheActor::GetBoundingBox() const
+	{
+		if (SkeletalMeshComponent)
+		{
+			return SkeletalMeshComponent->Bounds.GetBox();
+		}
+
+		if (GeomCacheComponent)
+		{
+			return GeomCacheComponent->Bounds.GetBox();
+		}
+
+		FBox Box;
+		Box.Init();
+		return Box;
+	}
+}	// namespace UE::MLDeformer
diff --git a/Engine/Plugins/Experimental/MLDeformer/MLDeformerFramework/Source/MLDeformerFrameworkEditor/Private/SMLDeformerTimeline.h b/Engine/Plugins/Experimental/MLDeformer/MLDeformerFramework/Source/MLDeformerFrameworkEditor/Private/SMLDeformerTimeline.h
--- a/Engine/Plugins/Experimental/MLDeformer/MLDeformerFramework/Source/MLDeformerFrameworkEditor/Private/SMLDeformerTimeline.h
+++ b/Engine/Plugins/Experimental/MLDeformer/MLDeformerFramework/Source/MLDeformerFrameworkEditor/Private/SMLDeformerTimeline.h
@@ -85,7 +85,7 @@ namespace UE::MLDeformer
 		TWeakPtr<FMLDeformerEditorModel> Model;
 
 		/** The currently active editor model. */
-		UE::MLDeformer::FMLDeformerEditorToolkit* Editor;
+		UE::MLDeformer::FMLDeformerEditorToolkit* Editor = nullptr;
 
 		/** The time slider controller. */
 		TSharedPtr<FMLTimeSliderController> TimeSliderController;
@@ -97,7 +97,7 @@ namespace UE::MLDeformer
 		TSharedPtr<ITimeSlider> TopTimeSlider;
 
 		/** The fill coefficients of each column in the grid. */
-		float ColumnFillCoefficients[2];
+		float ColumnFillCoefficients[2] = { 0.2f, 0.8f };
 
 		/** Called when the user has begun dragging the selection selection range. */
 		FSimpleDelegate OnSelectionRangeBeginDrag;
diff --git a/Engine/Plugins/Experimental/MLDeformer/MLDeformerFramework/Source/MLDeformerFrameworkEditor/Public/MLDeformerEditorActor.h b/Engine/Plugins/Experimental/MLDeformer/MLDeformerFramework/Source/MLDeformerFrameworkEditor/Public/MLDeformerEditorActor.h
--- a/Engine/Plugins/Experimental/MLDeformer/MLDeformerFramework/Source/MLDeformerFrameworkEditor/Public/MLDeformerEditorActor.h
+++ b/Engine/Plugins/Experimental/MLDeformer/MLDeformerFramework/Source/MLDeformerFrameworkEditor/Public/MLDeformerEditorActor.h
@@ -31,19 +31,21 @@ namespace UE::MLDeformer
 			FText LabelText;
 			bool bIsTrainingActor = false;
 		};
+
 		FMLDeformerEditorActor(const FConstructSettings& Settings);
 		virtual ~FMLDeformerEditorActor();
 
 		// Main methods you can override.
 		virtual void SetVisibility(bool bIsVisible);
 		virtual bool IsVisible() const;
-		virtual void SetPlayPosition(float TimeInSeconds, bool bAutoPause=true);
+		virtual void SetPlayPosition(float TimeInSeconds, bool bAutoPause = true);
 		virtual float GetPlayPosition() const;
 		virtual void SetPlaySpeed(float PlaySpeed);
 		virtual float GetPlaySpeed();
 
 		virtual FBox GetBoundingBox() const;
 		virtual void Pause(bool bPaused);
+		virtual bool IsPlaying() const;
 		virtual bool IsGroundTruthActor() const;
 		virtual bool HasVisualMesh() const;
 
@@ -88,4 +90,4 @@ namespace UE::MLDeformer
 		/** Is this actor used for training? */
 		bool bIsTrainingActor = true;
 	};
-}	// namespace UE::MLDeformer
+} // namespace UE::MLDeformer
diff --git a/Engine/Plugins/Experimental/MLDeformer/MLDeformerFramework/Source/MLDeformerFrameworkEditor/Public/MLDeformerEditorModel.h b/Engine/Plugins/Experimental/MLDeformer/MLDeformerFramework/Source/MLDeformerFrameworkEditor/Public/MLDeformerEditorModel.h
--- a/Engine/Plugins/Experimental/MLDeformer/MLDeformerFramework/Source/MLDeformerFrameworkEditor/Public/MLDeformerEditorModel.h
+++ b/Engine/Plugins/Experimental/MLDeformer/MLDeformerFramework/Source/MLDeformerFrameworkEditor/Public/MLDeformerEditorModel.h
@@ -52,12 +52,13 @@ enum class ETrainingResult : uint8
 namespace UE::MLDeformer
 {
 	class FMLDeformerEditorToolkit;
-	class FMLDeformerEditorActor;
 	class FMLDeformerSampler;
+	class FMLDeformerEditorModel;
 
 	/** The base class for the editor side of an UMLDeformerModel. */
 	class MLDEFORMERFRAMEWORKEDITOR_API FMLDeformerEditorModel
-		: public TSharedFromThis<FMLDeformerEditorModel>, public FGCObject
+		: public TSharedFromThis<FMLDeformerEditorModel>
+		, public FGCObject
 	{
 	public:
 		/** 
@@ -258,6 +259,7 @@ namespace UE::MLDeformer
 		double CalcTimelinePosition() const;
 
 		void UpdateRanges();
+
 	protected:
 		/** The runtime model associated with this editor model. */
 		TObjectPtr<UMLDeformerModel> Model = nullptr;
diff --git a/Engine/Plugins/Experimental/MLDeformer/MLDeformerFramework/Source/MLDeformerFrameworkEditor/Public/MLDeformerGeomCacheActor.h b/Engine/Plugins/Experimental/MLDeformer/MLDeformerFramework/Source/MLDeformerFrameworkEditor/Public/MLDeformerGeomCacheActor.h
new file mode 100644
--- /dev/null
+++ b/Engine/Plugins/Experimental/MLDeformer/MLDeformerFramework/Source/MLDeformerFrameworkEditor/Public/MLDeformerGeomCacheActor.h
@@ -0,0 +1,38 @@
+// Copyright Epic Games, Inc. All Rights Reserved.
+
+#pragma once
+
+#include "MLDeformerEditorActor.h"
+
+class UMLDeformerComponent;
+class UGeometryCacheComponent;
+
+namespace UE::MLDeformer
+{
+	class MLDEFORMERFRAMEWORKEDITOR_API FMLDeformerGeomCacheActor
+		: public FMLDeformerEditorActor
+	{
+	public:
+		FMLDeformerGeomCacheActor(const FConstructSettings& Settings);
+		virtual ~FMLDeformerGeomCacheActor() override;
+
+		void SetGeometryCacheComponent(UGeometryCacheComponent* Component) { GeomCacheComponent = Component; }
+		UGeometryCacheComponent* GetGeometryCacheComponent() const { return GeomCacheComponent; }
+
+		// FMLDeformerEditorActor overrides.
+		virtual void SetVisibility(bool bIsVisible) override;
+		virtual bool IsVisible() const override;
+		virtual bool HasVisualMesh() const override;
+		virtual void SetPlayPosition(float TimeInSeconds, bool bAutoPause = true) override;
+		virtual float GetPlayPosition() const override;
+		virtual void SetPlaySpeed(float PlaySpeed) override;
+		virtual void Pause(bool bPaused) override;
+		virtual bool IsPlaying() const override;
+		virtual FBox GetBoundingBox() const override;
+		// ~END FMLDeformerEditorActor overrides.
+
+	protected:
+		/** The geometry cache component (can be nullptr). */
+		TObjectPtr<UGeometryCacheComponent> GeomCacheComponent = nullptr;
+	};
+}	// namespace UE::MLDeformer
diff --git a/Engine/Plugins/Experimental/MLDeformer/NeuralMorphModel/Source/NeuralMorphModelEditor/Private/NeuralMorphEditorModel.cpp b/Engine/Plugins/Experimental/MLDeformer/NeuralMorphModel/Source/NeuralMorphModelEditor/Private/NeuralMorphEditorModel.cpp
--- a/Engine/Plugins/Experimental/MLDeformer/NeuralMorphModel/Source/NeuralMorphModelEditor/Private/NeuralMorphEditorModel.cpp
+++ b/Engine/Plugins/Experimental/MLDeformer/NeuralMorphModel/Source/NeuralMorphModelEditor/Private/NeuralMorphEditorModel.cpp
@@ -5,7 +5,7 @@
 #include "NeuralMorphModel.h"
 #include "NeuralMorphModelVizSettings.h"
 #include "NeuralMorphTrainingModel.h"
-#include "NeuralMorphEditorModelActor.h"
+#include "MLDeformerGeomCacheActor.h"
 #include "MLDeformerEditorToolkit.h"
 #include "MLDeformerEditorStyle.h"
 #include "MLDeformerInputInfo.h"
@@ -111,6 +111,7 @@ namespace UE::NeuralMorphModel
 
 		// Update the test geometry cache (ground truth) component.
 		GeometryCacheComponent = FindNeuralMorphModelEditorActor(ActorID_Test_GroundTruth)->GetGeometryCacheComponent();
+		check(GeometryCacheComponent);
 		GeometryCacheComponent->SetGeometryCache(VizSettings->GetTestGroundTruth());
 		GeometryCacheComponent->SetLooping(true);
 		GeometryCacheComponent->SetManualTick(true);
@@ -177,19 +178,19 @@ namespace UE::NeuralMorphModel
 		Settings.LabelColor = LabelColor;
 		Settings.LabelText = LabelText;
 		Settings.bIsTrainingActor = bIsTrainingActor;
-		FNeuralMorphEditorModelActor* EditorActor = static_cast<FNeuralMorphEditorModelActor*>(CreateEditorActor(Settings));
+		FMLDeformerGeomCacheActor* EditorActor = static_cast<FMLDeformerGeomCacheActor*>(CreateEditorActor(Settings));
 		EditorActor->SetGeometryCacheComponent(GeomCacheComponent);
 		EditorActors.Add(EditorActor);
 	} 
 
 	FMLDeformerEditorActor* FNeuralMorphEditorModel::CreateEditorActor(const FMLDeformerEditorActor::FConstructSettings& Settings) const
 	{
-		return new FNeuralMorphEditorModelActor(Settings);
+		return new FMLDeformerGeomCacheActor(Settings);
 	}
 
-	FNeuralMorphEditorModelActor* FNeuralMorphEditorModel::FindNeuralMorphModelEditorActor(int32 TypeID) const
+	FMLDeformerGeomCacheActor* FNeuralMorphEditorModel::FindNeuralMorphModelEditorActor(int32 TypeID) const
 	{
-		return static_cast<FNeuralMorphEditorModelActor*>(FindEditorActor(TypeID));
+		return static_cast<FMLDeformerGeomCacheActor*>(FindEditorActor(TypeID));
 	}
 
 	UNeuralMorphModelVizSettings* FNeuralMorphEditorModel::GetNeuralMorphModelVizSettings() const
@@ -200,7 +201,7 @@ namespace UE::NeuralMorphModel
 	double FNeuralMorphEditorModel::GetTrainingTimeAtFrame(int32 FrameNumber) const
 	{
 		// Try to get the frame from the geometry cache.
-		FNeuralMorphEditorModelActor* EditorActor = static_cast<FNeuralMorphEditorModelActor*>(FindEditorActor(ActorID_Train_GroundTruth));
+		FMLDeformerGeomCacheActor* EditorActor = static_cast<FMLDeformerGeomCacheActor*>(FindEditorActor(ActorID_Train_GroundTruth));
 		if (EditorActor && EditorActor->GetGeometryCacheComponent() && EditorActor->GetGeometryCacheComponent()->GeometryCache.Get())
 		{
 			return EditorActor->GetGeometryCacheComponent()->GetTimeAtFrame(FrameNumber);
@@ -211,7 +212,7 @@ namespace UE::NeuralMorphModel
 
 	int32 FNeuralMorphEditorModel::GetTrainingFrameAtTime(double TimeInSeconds) const
 	{
-		FNeuralMorphEditorModelActor* EditorActor = static_cast<FNeuralMorphEditorModelActor*>(FindEditorActor(ActorID_Train_GroundTruth));
+		FMLDeformerGeomCacheActor* EditorActor = static_cast<FMLDeformerGeomCacheActor*>(FindEditorActor(ActorID_Train_GroundTruth));
 		if (EditorActor && EditorActor->GetGeometryCacheComponent() && EditorActor->GetGeometryCacheComponent()->GeometryCache.Get())
 		{
 			return EditorActor->GetGeometryCacheComponent()->GetFrameAtTime(TimeInSeconds);
@@ -223,7 +224,7 @@ namespace UE::NeuralMorphModel
 	double FNeuralMorphEditorModel::GetTestTimeAtFrame(int32 FrameNumber) const
 	{
 		// Try to get the frame from the geometry cache.
-		FNeuralMorphEditorModelActor* EditorActor = static_cast<FNeuralMorphEditorModelActor*>(FindEditorActor(ActorID_Test_GroundTruth));
+		FMLDeformerGeomCacheActor* EditorActor = static_cast<FMLDeformerGeomCacheActor*>(FindEditorActor(ActorID_Test_GroundTruth));
 		if (EditorActor && EditorActor->GetGeometryCacheComponent() && EditorActor->GetGeometryCacheComponent()->GeometryCache.Get())
 		{
 			return EditorActor->GetGeometryCacheComponent()->GetTimeAtFrame(FrameNumber);
@@ -234,7 +235,7 @@ namespace UE::NeuralMorphModel
 
 	int32 FNeuralMorphEditorModel::GetTestFrameAtTime(double TimeInSeconds) const
 	{
-		FNeuralMorphEditorModelActor* EditorActor = static_cast<FNeuralMorphEditorModelActor*>(FindEditorActor(ActorID_Test_GroundTruth));
+		FMLDeformerGeomCacheActor* EditorActor = static_cast<FMLDeformerGeomCacheActor*>(FindEditorActor(ActorID_Test_GroundTruth));
 		if (EditorActor && EditorActor->GetGeometryCacheComponent() && EditorActor->GetGeometryCacheComponent()->GeometryCache.Get())
 		{
 			return EditorActor->GetGeometryCacheComponent()->GetFrameAtTime(TimeInSeconds);
@@ -374,7 +375,7 @@ namespace UE::NeuralMorphModel
 
 		// Debug draw the selected morph target.
 		UNeuralMorphModelVizSettings* VizSettings = GetNeuralMorphModelVizSettings();
-		if (VizSettings->bDrawMorphTargets)
+		if (VizSettings->bDrawMorphTargets && VizSettings->GetVisualizationMode() == EMLDeformerVizMode::TestData)
 		{
 			const FVector DrawOffset = -VizSettings->GetMeshSpacingOffsetVector();
 			DrawMorphTarget(PDI, GetNeuralMorphModel()->MorphTargetDeltas, VizSettings->MorphTargetDeltaThreshold, VizSettings->MorphTargetNumber, DrawOffset);
diff --git a/Engine/Plugins/Experimental/MLDeformer/NeuralMorphModel/Source/NeuralMorphModelEditor/Private/NeuralMorphEditorModelActor.cpp b/Engine/Plugins/Experimental/MLDeformer/NeuralMorphModel/Source/NeuralMorphModelEditor/Private/NeuralMorphEditorModelActor.cpp
deleted file mode 100644
diff --git a/Engine/Plugins/Experimental/MLDeformer/NeuralMorphModel/Source/NeuralMorphModelEditor/Public/NeuralMorphEditorModel.h b/Engine/Plugins/Experimental/MLDeformer/NeuralMorphModel/Source/NeuralMorphModelEditor/Public/NeuralMorphEditorModel.h
--- a/Engine/Plugins/Experimental/MLDeformer/NeuralMorphModel/Source/NeuralMorphModelEditor/Public/NeuralMorphEditorModel.h
+++ b/Engine/Plugins/Experimental/MLDeformer/NeuralMorphModel/Source/NeuralMorphModelEditor/Public/NeuralMorphEditorModel.h
@@ -9,6 +9,11 @@
 #include "NeuralMorphModel.h"
 #include "UObject/GCObject.h"
 
+namespace UE::MLDeformer
+{
+	class FMLDeformerGeomCacheActor;
+}
+
 class UMLDeformerAsset;
 class USkeletalMesh;
 class UGeometryCache;
@@ -60,9 +65,8 @@ namespace UE::NeuralMorphModel
 		// Some helpers that cast to this model's variants of some classes.
 		UNeuralMorphModel* GetNeuralMorphModel() const { return Cast<UNeuralMorphModel>(Model); }
 		UNeuralMorphModelVizSettings* GetNeuralMorphModelVizSettings() const;
-		FNeuralMorphEditorModelActor* FindNeuralMorphModelEditorActor(int32 TypeID) const;
+		FMLDeformerGeomCacheActor* FindNeuralMorphModelEditorActor(int32 TypeID) const;
 		FMLDeformerGeomCacheSampler* GetGeomCacheSampler() const { return static_cast<FMLDeformerGeomCacheSampler*>(Sampler); }
-	
 
 	protected:
 		void CreateGeomCacheActor(UWorld* World, int32 ActorID, const FName& Name, UGeometryCache* GeomCache, FLinearColor LabelColor, FLinearColor WireframeColor, const FText& LabelText, bool bIsTrainingActor);
diff --git a/Engine/Plugins/Experimental/MLDeformer/NeuralMorphModel/Source/NeuralMorphModelEditor/Public/NeuralMorphEditorModelActor.h b/Engine/Plugins/Experimental/MLDeformer/NeuralMorphModel/Source/NeuralMorphModelEditor/Public/NeuralMorphEditorModelActor.h
deleted file mode 100644
diff --git a/Engine/Plugins/Experimental/MLDeformer/NeuralMorphModel/Source/NeuralMorphModelEditor/Public/NeuralMorphModelDetails.h b/Engine/Plugins/Experimental/MLDeformer/NeuralMorphModel/Source/NeuralMorphModelEditor/Public/NeuralMorphModelDetails.h
--- a/Engine/Plugins/Experimental/MLDeformer/NeuralMorphModel/Source/NeuralMorphModelEditor/Public/NeuralMorphModelDetails.h
+++ b/Engine/Plugins/Experimental/MLDeformer/NeuralMorphModel/Source/NeuralMorphModelEditor/Public/NeuralMorphModelDetails.h
@@ -20,7 +20,7 @@ namespace UE::NeuralMorphModel
 {
 	class FNeuralMorphEditorModel;
 
-	class NEURALMORPHMODELEDITOR_API FNeuralMorphModelDetails final
+	class NEURALMORPHMODELEDITOR_API FNeuralMorphModelDetails
 		: public UE::MLDeformer::FMLDeformerModelDetails
 	{
 	public:
diff --git a/Engine/Plugins/Experimental/MLDeformer/NeuralMorphModel/Source/NeuralMorphModelEditor/Public/NeuralMorphModelVizSettingsDetails.h b/Engine/Plugins/Experimental/MLDeformer/NeuralMorphModel/Source/NeuralMorphModelEditor/Public/NeuralMorphModelVizSettingsDetails.h
--- a/Engine/Plugins/Experimental/MLDeformer/NeuralMorphModel/Source/NeuralMorphModelEditor/Public/NeuralMorphModelVizSettingsDetails.h
+++ b/Engine/Plugins/Experimental/MLDeformer/NeuralMorphModel/Source/NeuralMorphModelEditor/Public/NeuralMorphModelVizSettingsDetails.h
@@ -15,7 +15,7 @@ class UNeuralMorphModelVizSettings;
 
 namespace UE::NeuralMorphModel
 {
-	class NEURALMORPHMODELEDITOR_API FNeuralMorphModelVizSettingsDetails final
+	class NEURALMORPHMODELEDITOR_API FNeuralMorphModelVizSettingsDetails
 		: public UE::MLDeformer::FMLDeformerVizSettingsDetails
 	{
 	public:
diff --git a/Engine/Plugins/Experimental/MLDeformer/VertexDeltaModel/Source/VertexDeltaModelEditor/Private/VertexDeltaEditorModel.cpp b/Engine/Plugins/Experimental/MLDeformer/VertexDeltaModel/Source/VertexDeltaModelEditor/Private/VertexDeltaEditorModel.cpp
--- a/Engine/Plugins/Experimental/MLDeformer/VertexDeltaModel/Source/VertexDeltaModelEditor/Private/VertexDeltaEditorModel.cpp
+++ b/Engine/Plugins/Experimental/MLDeformer/VertexDeltaModel/Source/VertexDeltaModelEditor/Private/VertexDeltaEditorModel.cpp
@@ -4,7 +4,7 @@
 #include "VertexDeltaModel.h"
 #include "VertexDeltaModelVizSettings.h"
 #include "VertexDeltaTrainingModel.h"
-#include "VertexDeltaEditorModelActor.h"
+#include "MLDeformerGeomCacheActor.h"
 #include "MLDeformerEditorToolkit.h"
 #include "MLDeformerEditorStyle.h"
 #include "MLDeformerInputInfo.h"
@@ -157,19 +157,20 @@ namespace UE::VertexDeltaModel
 		Settings.LabelColor = LabelColor;
 		Settings.LabelText = LabelText;
 		Settings.bIsTrainingActor = bIsTrainingActor;
-		FVertexDeltaEditorModelActor* EditorActor = static_cast<FVertexDeltaEditorModelActor*>(CreateEditorActor(Settings));
+
+		FMLDeformerGeomCacheActor* EditorActor = static_cast<FMLDeformerGeomCacheActor*>(CreateEditorActor(Settings));
 		EditorActor->SetGeometryCacheComponent(GeomCacheComponent);
 		EditorActors.Add(EditorActor);
 	} 
 
 	FMLDeformerEditorActor* FVertexDeltaEditorModel::CreateEditorActor(const FMLDeformerEditorActor::FConstructSettings& Settings) const
 	{
-		return new FVertexDeltaEditorModelActor(Settings);
+		return new FMLDeformerGeomCacheActor(Settings);
 	}
 
-	FVertexDeltaEditorModelActor* FVertexDeltaEditorModel::FindVertexDeltaModelEditorActor(int32 TypeID) const
+	FMLDeformerGeomCacheActor* FVertexDeltaEditorModel::FindVertexDeltaModelEditorActor(int32 TypeID) const
 	{
-		return static_cast<FVertexDeltaEditorModelActor*>(FindEditorActor(TypeID));
+		return static_cast<FMLDeformerGeomCacheActor*>(FindEditorActor(TypeID));
 	}
 
 	UVertexDeltaModelVizSettings* FVertexDeltaEditorModel::GetVertexDeltaModelVizSettings() const
@@ -179,7 +180,7 @@ namespace UE::VertexDeltaModel
 
 	double FVertexDeltaEditorModel::GetTrainingTimeAtFrame(int32 FrameNumber) const
 	{
-		const FVertexDeltaEditorModelActor* EditorActor = static_cast<FVertexDeltaEditorModelActor*>(FindEditorActor(ActorID_Train_GroundTruth));
+		const FMLDeformerGeomCacheActor* EditorActor = static_cast<FMLDeformerGeomCacheActor*>(FindEditorActor(ActorID_Train_GroundTruth));
 		if (EditorActor && EditorActor->GetGeometryCacheComponent() && EditorActor->GetGeometryCacheComponent()->GeometryCache.Get())
 		{
 			return EditorActor->GetGeometryCacheComponent()->GetTimeAtFrame(FrameNumber);
@@ -190,7 +191,7 @@ namespace UE::VertexDeltaModel
 
 	int32 FVertexDeltaEditorModel::GetTrainingFrameAtTime(double TimeInSeconds) const
 	{
-		const FVertexDeltaEditorModelActor* EditorActor = static_cast<FVertexDeltaEditorModelActor*>(FindEditorActor(ActorID_Train_GroundTruth));
+		const FMLDeformerGeomCacheActor* EditorActor = static_cast<FMLDeformerGeomCacheActor*>(FindEditorActor(ActorID_Train_GroundTruth));
 		if (EditorActor && EditorActor->GetGeometryCacheComponent() && EditorActor->GetGeometryCacheComponent()->GeometryCache.Get())
 		{
 			return EditorActor->GetGeometryCacheComponent()->GetFrameAtTime(TimeInSeconds);
@@ -201,7 +202,7 @@ namespace UE::VertexDeltaModel
 
 	double FVertexDeltaEditorModel::GetTestTimeAtFrame(int32 FrameNumber) const
 	{
-		const FVertexDeltaEditorModelActor* EditorActor = static_cast<FVertexDeltaEditorModelActor*>(FindEditorActor(ActorID_Test_GroundTruth));
+		const FMLDeformerGeomCacheActor* EditorActor = static_cast<FMLDeformerGeomCacheActor*>(FindEditorActor(ActorID_Test_GroundTruth));
 		if (EditorActor && EditorActor->GetGeometryCacheComponent() && EditorActor->GetGeometryCacheComponent()->GeometryCache.Get())
 		{
 			return EditorActor->GetGeometryCacheComponent()->GetTimeAtFrame(FrameNumber);
@@ -212,7 +213,7 @@ namespace UE::VertexDeltaModel
 
 	int32 FVertexDeltaEditorModel::GetTestFrameAtTime(double TimeInSeconds) const
 	{
-		const FVertexDeltaEditorModelActor* EditorActor = static_cast<FVertexDeltaEditorModelActor*>(FindEditorActor(ActorID_Test_GroundTruth));
+		const FMLDeformerGeomCacheActor* EditorActor = static_cast<FMLDeformerGeomCacheActor*>(FindEditorActor(ActorID_Test_GroundTruth));
 		if (EditorActor && EditorActor->GetGeometryCacheComponent() && EditorActor->GetGeometryCacheComponent()->GeometryCache.Get())
 		{
 			return EditorActor->GetGeometryCacheComponent()->GetFrameAtTime(TimeInSeconds);
diff --git a/Engine/Plugins/Experimental/MLDeformer/VertexDeltaModel/Source/VertexDeltaModelEditor/Private/VertexDeltaEditorModelActor.cpp b/Engine/Plugins/Experimental/MLDeformer/VertexDeltaModel/Source/VertexDeltaModelEditor/Private/VertexDeltaEditorModelActor.cpp
deleted file mode 100644
diff --git a/Engine/Plugins/Experimental/MLDeformer/VertexDeltaModel/Source/VertexDeltaModelEditor/Public/VertexDeltaEditorModel.h b/Engine/Plugins/Experimental/MLDeformer/VertexDeltaModel/Source/VertexDeltaModelEditor/Public/VertexDeltaEditorModel.h
--- a/Engine/Plugins/Experimental/MLDeformer/VertexDeltaModel/Source/VertexDeltaModelEditor/Public/VertexDeltaEditorModel.h
+++ b/Engine/Plugins/Experimental/MLDeformer/VertexDeltaModel/Source/VertexDeltaModelEditor/Public/VertexDeltaEditorModel.h
@@ -9,6 +9,11 @@
 #include "MLDeformerEditorActor.h"
 #include "UObject/GCObject.h"
 
+namespace UE::MLDeformer
+{
+	class FMLDeformerGeomCacheActor;
+}
+
 class UMLDeformerAsset;
 class USkeletalMesh;
 class UGeometryCache;
@@ -19,7 +24,6 @@ namespace UE::VertexDeltaModel
 {
 	using namespace UE::MLDeformer;
 
-	class FVertexDeltaEditorModelActor;
 	class FVertexDeltaModelSampler;
 
 	class VERTEXDELTAMODELEDITOR_API FVertexDeltaEditorModel 
@@ -54,7 +58,7 @@ namespace UE::VertexDeltaModel
 		// Some helpers that cast to this model's variants of some classes.
 		UVertexDeltaModel* GetVertexDeltaModel() const { return Cast<UVertexDeltaModel>(Model); }
 		UVertexDeltaModelVizSettings* GetVertexDeltaModelVizSettings() const;
-		FVertexDeltaEditorModelActor* FindVertexDeltaModelEditorActor(int32 TypeID) const;
+		FMLDeformerGeomCacheActor* FindVertexDeltaModelEditorActor(int32 TypeID) const;
 		FMLDeformerGeomCacheSampler* GetGeomCacheSampler() const { return static_cast<FMLDeformerGeomCacheSampler*>(Sampler); }
 
 	protected:
diff --git a/Engine/Plugins/Experimental/MLDeformer/VertexDeltaModel/Source/VertexDeltaModelEditor/Public/VertexDeltaEditorModelActor.h b/Engine/Plugins/Experimental/MLDeformer/VertexDeltaModel/Source/VertexDeltaModelEditor/Public/VertexDeltaEditorModelActor.h
deleted file mode 100644
diff --git a/Engine/Plugins/Experimental/MLDeformer/VertexDeltaModel/Source/VertexDeltaModelEditor/Public/VertexDeltaModelDetails.h b/Engine/Plugins/Experimental/MLDeformer/VertexDeltaModel/Source/VertexDeltaModelEditor/Public/VertexDeltaModelDetails.h
--- a/Engine/Plugins/Experimental/MLDeformer/VertexDeltaModel/Source/VertexDeltaModelEditor/Public/VertexDeltaModelDetails.h
+++ b/Engine/Plugins/Experimental/MLDeformer/VertexDeltaModel/Source/VertexDeltaModelEditor/Public/VertexDeltaModelDetails.h
@@ -20,7 +20,7 @@ namespace UE::VertexDeltaModel
 {
 	class FVertexDeltaEditorModel;
 
-	class VERTEXDELTAMODELEDITOR_API FVertexDeltaModelDetails final
+	class VERTEXDELTAMODELEDITOR_API FVertexDeltaModelDetails
 		: public UE::MLDeformer::FMLDeformerModelDetails
 	{
 	public:
diff --git a/Engine/Plugins/Experimental/MLDeformer/VertexDeltaModel/Source/VertexDeltaModelEditor/Public/VertexDeltaModelVizSettingsDetails.h b/Engine/Plugins/Experimental/MLDeformer/VertexDeltaModel/Source/VertexDeltaModelEditor/Public/VertexDeltaModelVizSettingsDetails.h
--- a/Engine/Plugins/Experimental/MLDeformer/VertexDeltaModel/Source/VertexDeltaModelEditor/Public/VertexDeltaModelVizSettingsDetails.h
+++ b/Engine/Plugins/Experimental/MLDeformer/VertexDeltaModel/Source/VertexDeltaModelEditor/Public/VertexDeltaModelVizSettingsDetails.h
@@ -15,7 +15,7 @@ class UVertexDeltaModelVizSettings;
 
 namespace UE::VertexDeltaModel
 {
-	class VERTEXDELTAMODELEDITOR_API FVertexDeltaModelVizSettingsDetails final
+	class VERTEXDELTAMODELEDITOR_API FVertexDeltaModelVizSettingsDetails
 		: public UE::MLDeformer::FMLDeformerVizSettingsDetails
 	{
 	public:
