[MLDeformer]
- Added learning rate decay setting to the neural morph model.
- If you abort training, the samples don't have to be regenerated again next time.
- Changed learning rate to 0.001 and lower compresion settings and change default iteration count.

diff --git a/Engine/Plugins/Animation/MLDeformer/MLDeformerFramework/Source/MLDeformerFramework/Public/MLDeformerMorphModel.h b/Engine/Plugins/Animation/MLDeformer/MLDeformerFramework/Source/MLDeformerFramework/Public/MLDeformerMorphModel.h
--- a/Engine/Plugins/Animation/MLDeformer/MLDeformerFramework/Source/MLDeformerFramework/Public/MLDeformerMorphModel.h
+++ b/Engine/Plugins/Animation/MLDeformer/MLDeformerFramework/Source/MLDeformerFramework/Public/MLDeformerMorphModel.h
@@ -102,10 +102,10 @@ private:
 	 * A value of 0 will result in the highest quality morph targets, at the cost of higher runtime memory usage.
 	 */
 	UPROPERTY(EditAnywhere, Category = "Morph Targets", meta = (ClampMin = "0.0", ClampMax = "1.0", ForceUnits="cm"))
-	float MorphTargetDeltaThreshold = 0.01f;
+	float MorphTargetDeltaThreshold = 0.0025f;
 
 	/** The morph target error tolerance. Higher values result in larger compression, but could result in visual artifacts. */
 	UPROPERTY(EditAnywhere, Category = "Morph Targets", meta = (ClampMin = "0.01", ClampMax = "500"))
-	float MorphTargetErrorTolerance = 50.0f;
+	float MorphTargetErrorTolerance = 20.0f;
 #endif // WITH_EDITORONLY_DATA
 };
diff --git a/Engine/Plugins/Animation/MLDeformer/MLDeformerFramework/Source/MLDeformerFrameworkEditor/Private/MLDeformerEditorToolkit.cpp b/Engine/Plugins/Animation/MLDeformer/MLDeformerFramework/Source/MLDeformerFrameworkEditor/Private/MLDeformerEditorToolkit.cpp
--- a/Engine/Plugins/Animation/MLDeformer/MLDeformerFramework/Source/MLDeformerFrameworkEditor/Private/MLDeformerEditorToolkit.cpp
+++ b/Engine/Plugins/Animation/MLDeformer/MLDeformerFramework/Source/MLDeformerFrameworkEditor/Private/MLDeformerEditorToolkit.cpp
@@ -498,6 +498,7 @@ namespace UE::MLDeformer
 			// User aborted the training process. Ask whether they want to use the partially trained results or not.
 			case ETrainingResult::Aborted:
 			{
+				ActiveModel->SetResamplingInputOutputsNeeded(false);
 				const FText Title(LOCTEXT("TrainingAbortedMessageTitle", "Use partially trained network?"));
 				const EAppReturnType::Type ReturnType = FMessageDialog::Open(
 					EAppMsgType::YesNo, 
diff --git a/Engine/Plugins/Animation/MLDeformer/MLDeformerFramework/Source/MLDeformerFrameworkEditor/Public/MLDeformerEditorModel.h b/Engine/Plugins/Animation/MLDeformer/MLDeformerFramework/Source/MLDeformerFrameworkEditor/Public/MLDeformerEditorModel.h
--- a/Engine/Plugins/Animation/MLDeformer/MLDeformerFramework/Source/MLDeformerFrameworkEditor/Public/MLDeformerEditorModel.h
+++ b/Engine/Plugins/Animation/MLDeformer/MLDeformerFramework/Source/MLDeformerFrameworkEditor/Public/MLDeformerEditorModel.h
@@ -940,8 +940,9 @@ namespace UE::MLDeformer
 	};
 
 	/**
-	 * 
-	 * 
+	 * Train the model.
+	 * This will execute the train function inside your python code.
+	 * Provide the UMLDeformerTrainingModel inherited class type as template argument.
 	 */
 	template<class TrainingModelClass>
 	ETrainingResult TrainModel(FMLDeformerEditorModel* EditorModel)
diff --git a/Engine/Plugins/Animation/MLDeformer/NeuralMorphModel/Source/NeuralMorphModel/Public/NeuralMorphModel.h b/Engine/Plugins/Animation/MLDeformer/NeuralMorphModel/Source/NeuralMorphModel/Public/NeuralMorphModel.h
--- a/Engine/Plugins/Animation/MLDeformer/NeuralMorphModel/Source/NeuralMorphModel/Public/NeuralMorphModel.h
+++ b/Engine/Plugins/Animation/MLDeformer/NeuralMorphModel/Source/NeuralMorphModel/Public/NeuralMorphModel.h
@@ -73,6 +73,7 @@ public:
 	int32 GetNumIterations() const						{ return NumIterations; }
 	int32 GetBatchSize() const							{ return BatchSize; }
 	float GetLearningRate() const						{ return LearningRate; }
+	float GetLearningRateDecay() const					{ return LearningRateDecay; }
 	float GetRegularizationFactor() const				{ return RegularizationFactor;  }
 
 public:
@@ -105,7 +106,7 @@ public:
 	 * Once the loss doesn't go down anymore, you know that more iterations most likely won't help much.
 	 */
 	UPROPERTY(EditAnywhere, BlueprintReadWrite, Category = "Training Settings", meta = (ClampMin = "1", ClampMax = "1000000"))
-	int32 NumIterations = 2000;
+	int32 NumIterations = 5000;
 
 	/** 
 	 * The number of hidden layers that the neural network model will have.
@@ -137,7 +138,16 @@ public:
 
 	/** The learning rate used during the model training. */
 	UPROPERTY(EditAnywhere, BlueprintReadWrite, AdvancedDisplay, Category = "Training Settings", meta = (ClampMin = "0.000001", ClampMax = "1.0"))
-	float LearningRate = 0.01f;
+	float LearningRate = 0.001f;
+
+	/** 
+	 * The learning rate decay rate. If this is set to 1, the learning rate will not change. 
+	 * This is a multiplication factor. Every 1000 iterations, the new learning rate will be equal to
+	 * CurrentLearningRate * LearningRateDecay. This allows us to take slightly larger steps in the first iterations, and slowly take smaller steps.
+	 * Generally you want this to be between 0.9 and 1.
+	 */
+	UPROPERTY(EditAnywhere, BlueprintReadWrite, AdvancedDisplay, Category = "Training Settings", meta = (ClampMin = "0.1", ClampMax = "1.0"))
+	float LearningRateDecay = 0.98f;
 
 	/** 
 	 * The regularization factor. Higher values can help generate more sparse morph targets, but can also lead to visual artifacts. 
diff --git a/Engine/Plugins/Animation/MLDeformer/NeuralMorphModel/Source/NeuralMorphModelEditor/Private/NeuralMorphModelDetails.cpp b/Engine/Plugins/Animation/MLDeformer/NeuralMorphModel/Source/NeuralMorphModelEditor/Private/NeuralMorphModelDetails.cpp
--- a/Engine/Plugins/Animation/MLDeformer/NeuralMorphModel/Source/NeuralMorphModelEditor/Private/NeuralMorphModelDetails.cpp
+++ b/Engine/Plugins/Animation/MLDeformer/NeuralMorphModel/Source/NeuralMorphModelEditor/Private/NeuralMorphModelDetails.cpp
@@ -45,6 +45,7 @@ namespace UE::NeuralMorphModel
 		TrainingSettingsCategoryBuilder->AddProperty(GET_MEMBER_NAME_CHECKED(UNeuralMorphModel, NumIterations), UNeuralMorphModel::StaticClass());
 		TrainingSettingsCategoryBuilder->AddProperty(GET_MEMBER_NAME_CHECKED(UNeuralMorphModel, BatchSize), UNeuralMorphModel::StaticClass());
 		TrainingSettingsCategoryBuilder->AddProperty(GET_MEMBER_NAME_CHECKED(UNeuralMorphModel, LearningRate), UNeuralMorphModel::StaticClass());
+		TrainingSettingsCategoryBuilder->AddProperty(GET_MEMBER_NAME_CHECKED(UNeuralMorphModel, LearningRateDecay), UNeuralMorphModel::StaticClass());
 		TrainingSettingsCategoryBuilder->AddProperty(GET_MEMBER_NAME_CHECKED(UNeuralMorphModel, RegularizationFactor), UNeuralMorphModel::StaticClass());
 	}
 }	// namespace UE::NeuralMorphModel
