[MLDeformer]
- Added versioning for ML Deformer assets.
- Input info doesn't store both FNames and FStrings anymore, but only FNames. Handled backward compatibility.
- The input info now contains a Skeletal Mesh it was trained on, making compatibility checks safer.
- ML Deformer component now searches for the Skeletal Mesh based on the skel mesh it was actually trained on, not the one picked in the asset editor.
- Added a Neural Morph Network class, used for custom inference.
- Adjusted things to deal with custom inference.
- Allow 0 hidden layers in neural morph model.
- Add a special compile mode to force use NNI.
- Some optimizations and refactors to how network bone rotation inputs are calculated.

diff --git a/Engine/Plugins/Animation/MLDeformer/MLDeformerFramework/Config/BaseMLDeformerFramework.ini b/Engine/Plugins/Animation/MLDeformer/MLDeformerFramework/Config/BaseMLDeformerFramework.ini
--- a/Engine/Plugins/Animation/MLDeformer/MLDeformerFramework/Config/BaseMLDeformerFramework.ini
+++ b/Engine/Plugins/Animation/MLDeformer/MLDeformerFramework/Config/BaseMLDeformerFramework.ini
@@ -1,4 +1,8 @@
 [CoreRedirects]
-; Redirects from migrationof MLDeformer to MLDeformerFramework
+; Redirects from MLDeformer to MLDeformerFramework.
 +ClassRedirects=(OldName="/Script/MLDeformer.MLDeformerAsset",NewName="/Script/MLDeformerFramework.MLDeformerAsset")
 +PackageRedirects=(OldName="MLDeformer",NewName="/MLDeformerFramework/",MatchSubstring=true)
+
+; Redirects from BoneNameStrings to BoneNames and from CurveNameStrings to CurveNames.
++PropertyRedirects=(OldName="/Script/MLDeformerFramework.MLDeformerInputInfo.BoneNameStrings",NewName="BoneNames")
++PropertyRedirects=(OldName="/Script/MLDeformerFramework.MLDeformerInputInfo.CurveNameStrings",NewName="CurveNames")
diff --git a/Engine/Plugins/Animation/MLDeformer/MLDeformerFramework/Source/MLDeformerFramework/Private/MLDeformerAsset.cpp b/Engine/Plugins/Animation/MLDeformer/MLDeformerFramework/Source/MLDeformerFramework/Private/MLDeformerAsset.cpp
new file mode 100644
--- /dev/null
+++ b/Engine/Plugins/Animation/MLDeformer/MLDeformerFramework/Source/MLDeformerFramework/Private/MLDeformerAsset.cpp
@@ -0,0 +1,12 @@
+// Copyright Epic Games, Inc. All Rights Reserved.
+
+#include "MLDeformerAsset.h"
+#include "MLDeformerObjectVersion.h"
+
+void UMLDeformerAsset::Serialize(FArchive& Archive)
+{
+	TRACE_CPUPROFILER_EVENT_SCOPE(UMLDeformerAsset::Serialize)
+
+	Archive.UsingCustomVersion(UE::MLDeformer::FMLDeformerObjectVersion::GUID);
+	Super::Serialize(Archive);
+}
diff --git a/Engine/Plugins/Animation/MLDeformer/MLDeformerFramework/Source/MLDeformerFramework/Private/MLDeformerComponent.cpp b/Engine/Plugins/Animation/MLDeformer/MLDeformerFramework/Source/MLDeformerFramework/Private/MLDeformerComponent.cpp
--- a/Engine/Plugins/Animation/MLDeformer/MLDeformerFramework/Source/MLDeformerFramework/Private/MLDeformerComponent.cpp
+++ b/Engine/Plugins/Animation/MLDeformer/MLDeformerFramework/Source/MLDeformerFramework/Private/MLDeformerComponent.cpp
@@ -4,9 +4,11 @@
 #include "MLDeformerAsset.h"
 #include "MLDeformerModelInstance.h"
 #include "MLDeformerModel.h"
+#include "MLDeformerInputInfo.h"
 #include "Components/SkeletalMeshComponent.h"
 #include "Animation/MeshDeformer.h"
 #include "Animation/MeshDeformerInstance.h"
+#include "Engine/SkeletalMesh.h"
 
 UMLDeformerComponent::UMLDeformerComponent(const FObjectInitializer& ObjectInitializer)
 	: Super(ObjectInitializer)
@@ -145,7 +147,7 @@ USkeletalMeshComponent* UMLDeformerComponent::FindSkeletalMeshComponent(UMLDefor
 		const UMLDeformerModel* Model = Asset ? Asset->GetModel() : nullptr;
 		if (Model && Model->GetSkeletalMesh())
 		{
-			const USkeletalMesh* ModelSkeletalMesh = Model->GetSkeletalMesh();
+			const FSoftObjectPath& ModelSkeletalMesh = Model->GetInputInfo()->GetSkeletalMesh();
 
 			// Get a list of all skeletal mesh components on the actor.
 			TArray<USkeletalMeshComponent*> Components;
@@ -158,7 +160,7 @@ USkeletalMeshComponent* UMLDeformerComponent::FindSkeletalMeshComponent(UMLDefor
 				for (USkeletalMeshComponent* Component : Components)
 				{
 					const USkeletalMesh* ComponentSkeletalMesh = Component->GetSkeletalMeshAsset();
-					if (ComponentSkeletalMesh == ModelSkeletalMesh)
+					if (FSoftObjectPath(ComponentSkeletalMesh) == ModelSkeletalMesh)
 					{
 						ResultingComponent = Component;
 						break;
@@ -190,7 +192,6 @@ void UMLDeformerComponent::UpdateSkeletalMeshComponent()
 void UMLDeformerComponent::Activate(bool bReset)
 {
 	UpdateSkeletalMeshComponent();
-	SetupComponent(DeformerAsset, SkelMeshComponent);
 	Super::Activate(bReset);
 }
 
diff --git a/Engine/Plugins/Animation/MLDeformer/MLDeformerFramework/Source/MLDeformerFramework/Private/MLDeformerInputInfo.cpp b/Engine/Plugins/Animation/MLDeformer/MLDeformerFramework/Source/MLDeformerFramework/Private/MLDeformerInputInfo.cpp
--- a/Engine/Plugins/Animation/MLDeformer/MLDeformerFramework/Source/MLDeformerFramework/Private/MLDeformerInputInfo.cpp
+++ b/Engine/Plugins/Animation/MLDeformer/MLDeformerFramework/Source/MLDeformerFramework/Private/MLDeformerInputInfo.cpp
@@ -3,6 +3,7 @@
 #include "MLDeformerInputInfo.h"
 #include "MLDeformerAsset.h"
 #include "MLDeformerModel.h"
+#include "MLDeformerObjectVersion.h"
 #include "Animation/AnimCurveTypes.h"
 #include "Animation/Skeleton.h"
 #include "Animation/AnimInstance.h"
@@ -10,46 +11,51 @@
 #include "Engine/SkeletalMesh.h"
 #include "Components/SkeletalMeshComponent.h"
 
-
 void UMLDeformerInputInfo::Reset()
 {
-	BoneNameStrings.Empty();
 	BoneNames.Empty();
-	CurveNameStrings.Empty();
 	CurveNames.Empty();
 	NumBaseMeshVertices = 0;
 	NumTargetMeshVertices = 0;
+	SkeletalMesh = nullptr;
 }
 
-void UMLDeformerInputInfo::UpdateFNames()
+void UMLDeformerInputInfo::OnPostLoad()
+{
+	PRAGMA_DISABLE_DEPRECATION_WARNINGS
+	UpdateNameStrings();
+	PRAGMA_ENABLE_DEPRECATION_WARNINGS
+}
+
+void UMLDeformerInputInfo::UpdateNameStrings()
 {
-	// Update the bone names.
-	const int32 NumBones = BoneNameStrings.Num();
-	BoneNames.Reset(NumBones);
-	BoneNames.Reserve(NumBones);
-	for (const FString& NameString : BoneNameStrings)
+#if WITH_EDITORONLY_DATA
+	BoneNameStrings_DEPRECATED.Reset();
+	BoneNameStrings_DEPRECATED.AddDefaulted(BoneNames.Num());
+	for (int32 Index = 0; Index < BoneNames.Num(); ++Index)
 	{
-		BoneNames.Add(FName(NameString));
+		BoneNameStrings_DEPRECATED[Index] = BoneNames[Index].ToString();
 	}
 
-	// Update the curve names.
-	const int32 NumCurves = CurveNameStrings.Num();
-	CurveNames.Reset(NumCurves);
-	CurveNames.Reserve(NumCurves);
-	for (const FString& NameString : CurveNameStrings)
+	CurveNameStrings_DEPRECATED.Reset();
+	CurveNameStrings_DEPRECATED.AddDefaulted(CurveNames.Num());
+	for (int32 Index = 0; Index < CurveNames.Num(); ++Index)
 	{
-		CurveNames.Add(FName(NameString));
+		CurveNameStrings_DEPRECATED[Index] = CurveNames[Index].ToString();
 	}
+#endif
 }
 
-void UMLDeformerInputInfo::OnPostLoad()
-{ 
-	UpdateFNames();
+void UMLDeformerInputInfo::UpdateFNames()
+{
+	PRAGMA_DISABLE_DEPRECATION_WARNINGS
+	UpdateNameStrings();
+	PRAGMA_ENABLE_DEPRECATION_WARNINGS
 }
 
 bool UMLDeformerInputInfo::IsEmpty() const
 { 
-	return (BoneNameStrings.IsEmpty() && CurveNameStrings.IsEmpty());
+	return (BoneNames.IsEmpty() && CurveNames.IsEmpty());
 }
 
 int32 UMLDeformerInputInfo::GetNumBones() const
@@ -57,12 +63,7 @@ int32 UMLDeformerInputInfo::GetNumBones() const
 	return BoneNames.Num();
 }
 
-const FString& UMLDeformerInputInfo::GetBoneNameString(int32 Index) const
-{ 
-	return BoneNameStrings[Index];
-}
-
-const FName UMLDeformerInputInfo::GetBoneName(int32 Index) const
+FName UMLDeformerInputInfo::GetBoneName(int32 Index) const
 { 
 	return BoneNames[Index];
 }
@@ -72,16 +73,6 @@ int32 UMLDeformerInputInfo::GetNumCurves() const
 	return CurveNames.Num();
 }
 
-TArray<FString>& UMLDeformerInputInfo::GetBoneNameStrings()
-{ 
-	return BoneNameStrings;
-}
-
-TArray<FString>& UMLDeformerInputInfo::GetCurveNameStrings()
-{ 
-	return CurveNameStrings;
-}
-
 TArray<FName>& UMLDeformerInputInfo::GetBoneNames()
 { 
 	return BoneNames;
@@ -92,12 +83,7 @@ TArray<FName>& UMLDeformerInputInfo::GetCurveNames()
 	return CurveNames;
 }
 
-const FString& UMLDeformerInputInfo::GetCurveNameString(int32 Index) const
-{ 
-	return CurveNameStrings[Index];
-}
-
-const FName UMLDeformerInputInfo::GetCurveName(int32 Index) const
+FName UMLDeformerInputInfo::GetCurveName(int32 Index) const
 { 
 	return CurveNames[Index];
 }
@@ -122,15 +108,20 @@ void UMLDeformerInputInfo::SetNumTargetVertices(int32 NumVerts)
 	NumTargetMeshVertices = NumVerts;
 }
 
-bool UMLDeformerInputInfo::IsCompatible(USkeletalMesh* SkeletalMesh) const
+bool UMLDeformerInputInfo::IsCompatible(USkeletalMesh* InSkeletalMesh) const
 {
-	if (SkeletalMesh == nullptr)
+	if (InSkeletalMesh == nullptr)
+	{
+		return false;
+	}
+
+	if (FSoftObjectPath(InSkeletalMesh) != SkeletalMesh)
 	{
 		return false;
 	}
 
 	// Verify that all required bones are there.
-	const FReferenceSkeleton& RefSkeleton = SkeletalMesh->GetRefSkeleton();
+	const FReferenceSkeleton& RefSkeleton = InSkeletalMesh->GetRefSkeleton();
 	for (const FName BoneName : BoneNames)
 	{
 		const int32 BoneIndex = RefSkeleton.FindBoneIndex(BoneName);
@@ -141,7 +132,7 @@ bool UMLDeformerInputInfo::IsCompatible(USkeletalMesh* SkeletalMesh) const
 	}
 
 	// Verify that all required curves are there.
-	USkeleton* Skeleton = SkeletalMesh->GetSkeleton();
+	USkeleton* Skeleton = InSkeletalMesh->GetSkeleton();
 	const FSmartNameMapping* SmartNameMapping = Skeleton ? Skeleton->GetSmartNameContainer(USkeleton::AnimCurveMappingName) : nullptr;
 	if (SmartNameMapping)
 	{
@@ -157,15 +148,15 @@ bool UMLDeformerInputInfo::IsCompatible(USkeletalMesh* SkeletalMesh) const
 	return true;
 }
 
-FString UMLDeformerInputInfo::GenerateCompatibilityErrorString(USkeletalMesh* SkeletalMesh) const
+FString UMLDeformerInputInfo::GenerateCompatibilityErrorString(USkeletalMesh* InSkeletalMesh) const
 {
-	if (SkeletalMesh == nullptr)
+	if (InSkeletalMesh == nullptr)
 	{
 		return FString();
 	}
 
 	// Verify that all required bones are there.
-	const FReferenceSkeleton& RefSkeleton = SkeletalMesh->GetRefSkeleton();
+	const FReferenceSkeleton& RefSkeleton = InSkeletalMesh->GetRefSkeleton();
 	FString ErrorString;
 	for (const FName BoneName : BoneNames)
 	{
@@ -177,7 +168,7 @@ FString UMLDeformerInputInfo::GenerateCompatibilityErrorString(USkeletalMesh* Sk
 	}
 
 	// Verify that all required curves are there.
-	USkeleton* Skeleton = SkeletalMesh->GetSkeleton();
+	USkeleton* Skeleton = InSkeletalMesh->GetSkeleton();
 	const FSmartNameMapping* SmartNameMapping = Skeleton ? Skeleton->GetSmartNameContainer(USkeleton::AnimCurveMappingName) : nullptr;
 	if (SmartNameMapping)
 	{
@@ -193,24 +184,27 @@ FString UMLDeformerInputInfo::GenerateCompatibilityErrorString(USkeletalMesh* Sk
 	// Check vertex count.
 #if WITH_EDITORONLY_DATA
 	if ((NumBaseMeshVertices > 0 && NumTargetMeshVertices > 0) &&
-		NumBaseMeshVertices != SkeletalMesh->GetNumImportedVertices())
+		NumBaseMeshVertices != InSkeletalMesh->GetNumImportedVertices())
 	{
 		ErrorString += FString::Format(TEXT("The number of vertices that the network was trained on ({0} verts) doesn't match the skeletal mesh '{1}' ({2} verts)..\n"), 
 			{
 				NumBaseMeshVertices, 
-				SkeletalMesh->GetName(),
-				SkeletalMesh->GetNumImportedVertices(),
+				InSkeletalMesh->GetName(),
+				InSkeletalMesh->GetNumImportedVertices(),
 			} );
 	}
 #endif
 
+	if (SkeletalMesh.IsValid() && FSoftObjectPath(InSkeletalMesh) != SkeletalMesh)
+	{
+		ErrorString += FString::Format(TEXT("Model was trained on Skeletal Mesh '{0}', which is not the same as '{1}' it tries to apply on.\n"), {*SkeletalMesh.ToString(), *FSoftObjectPath(InSkeletalMesh).ToString()});
+	}
+
 	return ErrorString;
 }
 
 void UMLDeformerInputInfo::ExtractCurveValues(USkeletalMeshComponent* SkelMeshComponent, TArray<float>& OutValues) const
 {
-	check(CurveNames.Num() == CurveNameStrings.Num());
-
 	UAnimInstance* AnimInstance = SkelMeshComponent->GetAnimInstance();
 	const int32 NumCurves = CurveNames.Num();
 	OutValues.Reset(NumCurves);
@@ -222,6 +216,58 @@ void UMLDeformerInputInfo::ExtractCurveValues(USkeletalMeshComponent* SkelMeshCo
 	}
 }
 
+void UMLDeformerInputInfo::RotationToTwoVectorsAsSixFloats(TArrayView<FTransform> Transforms, float* OutputBuffer)
+{
+	for (int32 Index = 0; Index < Transforms.Num(); ++Index)
+	{
+		// Precalculate some things, essentially like a quaternion to matrix conversion.
+		const FQuat Rotation = Transforms[Index].GetRotation();
+		const float X = Rotation.X;
+		const float Y = Rotation.Y;
+		const float Z = Rotation.Z;
+		const float W = Rotation.W;
+		const float x2 = X + X;    const float y2 = Y + Y;    const float z2 = Z + Z;
+		const float xx = X * x2;   const float xy = X * y2;   const float xz = X * z2;
+		const float yy = Y * y2;   const float yz = Y * z2;   const float zz = Z * z2;
+		const float wx = W * x2;   const float wy = W * y2;   const float wz = W * z2;
+
+		const int32 OutputOffset = Index * 6;
+
+		// X Column of the matrix.
+		OutputBuffer[OutputOffset + 0] = 1.0f - (yy + zz);
+		OutputBuffer[OutputOffset + 1] = xy - wz;
+		OutputBuffer[OutputOffset + 2] = xz + wy;
+
+		// Y Column of the matrix.
+		OutputBuffer[OutputOffset + 3] = xy + wz;
+		OutputBuffer[OutputOffset + 4] = 1.0f - (xx + zz);
+		OutputBuffer[OutputOffset + 5] = yz - wx;
+	}
+}
+
+void UMLDeformerInputInfo::RotationToTwoVectorsAsSixFloats(const FQuat& Rotation, float* SixFloatsOutputBuffer)
+{
+	// Precalculate some things, essentially like a quaternion to matrix conversion.
+	const float X = Rotation.X;
+	const float Y = Rotation.Y;
+	const float Z = Rotation.Z;
+	const float W = Rotation.W;
+	const float x2 = X + X;    const float y2 = Y + Y;    const float z2 = Z + Z;
+	const float xx = X * x2;   const float xy = X * y2;   const float xz = X * z2;
+	const float yy = Y * y2;   const float yz = Y * z2;   const float zz = Z * z2;
+	const float wx = W * x2;   const float wy = W * y2;   const float wz = W * z2;
+
+	// X Column of the matrix.
+	SixFloatsOutputBuffer[0] = 1.0f - (yy + zz);
+	SixFloatsOutputBuffer[1] = xy - wz;
+	SixFloatsOutputBuffer[2] = xz + wy;
+
+	// Y Column of the matrix.
+	SixFloatsOutputBuffer[3] = xy + wz;
+	SixFloatsOutputBuffer[4] = 1.0f - (xx + zz);
+	SixFloatsOutputBuffer[5] = yz - wx;
+}
+
 void UMLDeformerInputInfo::ExtractBoneRotations(USkeletalMeshComponent* SkelMeshComponent, TArray<float>& OutRotations) const
 {
 	const TArray<FTransform>& BoneTransforms = SkelMeshComponent->GetBoneSpaceTransforms();
@@ -229,26 +275,30 @@ void UMLDeformerInputInfo::ExtractBoneRotations(USkeletalMeshComponent* SkelMesh
 	const int32 NumFloats = NumBones * 6; // 2 Columns of the rotation matrix.
 	OutRotations.Reset(NumFloats);
 	OutRotations.AddUninitialized(NumFloats);
-	int32 Offset = 0;
 	for (int32 Index = 0; Index < NumBones; ++Index)
 	{
 		const FName BoneName = GetBoneName(Index);
 		const int32 SkelMeshBoneIndex = SkelMeshComponent->GetBoneIndex(BoneName);
-		const FMatrix RotationMatrix = (SkelMeshBoneIndex != INDEX_NONE) ? BoneTransforms[SkelMeshBoneIndex].GetRotation().ToMatrix() : FMatrix::Identity;
-		const FVector X = RotationMatrix.GetColumn(0);
-		const FVector Y = RotationMatrix.GetColumn(1);	
-		OutRotations[Offset++] = X.X;
-		OutRotations[Offset++] = X.Y;
-		OutRotations[Offset++] = X.Z;
-		OutRotations[Offset++] = Y.X;
-		OutRotations[Offset++] = Y.Y;
-		OutRotations[Offset++] = Y.Z;
+		const FQuat Rotation = (SkelMeshBoneIndex != INDEX_NONE) ? BoneTransforms[SkelMeshBoneIndex].GetRotation() : FQuat::Identity;
+		RotationToTwoVectorsAsSixFloats(
+			Rotation,
+			OutRotations.GetData() + Index * 6);
 	}
 }
 
 int32 UMLDeformerInputInfo::CalcNumNeuralNetInputs() const
 {
 	return 
-		BoneNameStrings.Num() * 6 +	// Six floats per bone (2 FVector3 columns of the rotation matrix).
-		CurveNameStrings.Num();		// One float per curve.
+		BoneNames.Num() * 6 +	// Six floats per bone (2 FVector3 columns of the rotation matrix).
+		CurveNames.Num();		// One float per curve.
+}
+
+const FSoftObjectPath& UMLDeformerInputInfo::GetSkeletalMesh() const
+{
+	return SkeletalMesh;
+}
+
+void UMLDeformerInputInfo::SetSkeletalMesh(USkeletalMesh* InSkeletalMesh)
+{
+	SkeletalMesh = InSkeletalMesh;
 }
diff --git a/Engine/Plugins/Animation/MLDeformer/MLDeformerFramework/Source/MLDeformerFramework/Private/MLDeformerModel.cpp b/Engine/Plugins/Animation/MLDeformer/MLDeformerFramework/Source/MLDeformerFramework/Private/MLDeformerModel.cpp
--- a/Engine/Plugins/Animation/MLDeformer/MLDeformerFramework/Source/MLDeformerFramework/Private/MLDeformerModel.cpp
+++ b/Engine/Plugins/Animation/MLDeformer/MLDeformerFramework/Source/MLDeformerFramework/Private/MLDeformerModel.cpp
@@ -6,6 +6,7 @@
 #include "MLDeformerVizSettings.h"
 #include "MLDeformerAsset.h"
 #include "MLDeformerComponent.h"
+#include "MLDeformerObjectVersion.h"
 #include "Engine/SkeletalMesh.h"
 #include "Rendering/SkeletalMeshModel.h"
 #include "Rendering/SkeletalMeshLODModel.h"
@@ -67,7 +68,6 @@ void UMLDeformerModel::Init(UMLDeformerAsset* InDeformerAsset)
 
 void UMLDeformerModel::Serialize(FArchive& Archive)
 {
-	TRACE_CPUPROFILER_EVENT_SCOPE(UMLDeformerModel::Serialize)
 	#if WITH_EDITOR
 		if (Archive.IsSaving())
 		{
@@ -112,6 +112,13 @@ void UMLDeformerModel::PostLoad()
 
 	if (InputInfo)
 	{
+		// If we are dealing with an input info that doesn't have a skeletal mesh, then use the 
+		// currently set skeletal mesh. This is for backward compatibility, from before we put the skeletal mesh inside the input info.
+		if (InputInfo->GetSkeletalMesh() == nullptr)
+		{
+			InputInfo->SetSkeletalMesh(SkeletalMesh);
+		}
+
 		InputInfo->OnPostLoad();
 	}
 
diff --git a/Engine/Plugins/Animation/MLDeformer/MLDeformerFramework/Source/MLDeformerFramework/Private/MLDeformerModelInstance.cpp b/Engine/Plugins/Animation/MLDeformer/MLDeformerFramework/Source/MLDeformerFramework/Private/MLDeformerModelInstance.cpp
--- a/Engine/Plugins/Animation/MLDeformer/MLDeformerFramework/Source/MLDeformerFramework/Private/MLDeformerModelInstance.cpp
+++ b/Engine/Plugins/Animation/MLDeformer/MLDeformerFramework/Source/MLDeformerFramework/Private/MLDeformerModelInstance.cpp
@@ -187,12 +187,12 @@ void UMLDeformerModelInstance::UpdateBoneTransforms()
 	{
 		// Grab the transforms from our own skeletal mesh component.
 		// These are local space transforms, relative to the parent bone.
-		BoneTransforms = SkeletalMeshComponent->GetBoneSpaceTransforms();
+		const TArrayView<const FTransform> Transforms = SkeletalMeshComponent->GetBoneSpaceTransformsView();
 		const int32 NumTrainingBones = AssetBonesToSkelMeshMappings.Num();
 		for (int32 Index = 0; Index < NumTrainingBones; ++Index)
 		{
 			const int32 ComponentBoneIndex = AssetBonesToSkelMeshMappings[Index];
-			TrainingBoneTransforms[Index] = BoneTransforms[ComponentBoneIndex];
+			TrainingBoneTransforms[Index] = Transforms[ComponentBoneIndex];
 		}
 	}
 }
@@ -203,28 +203,15 @@ int64 UMLDeformerModelInstance::SetBoneTransforms(float* OutputBuffer, int64 Out
 	// These are in the space relative to their parent.
 	UpdateBoneTransforms();
 
-	// Write the transforms into the output buffer.
-	const UMLDeformerInputInfo* InputInfo = Model->GetInputInfo();
-	const int32 AssetNumBones = InputInfo->GetNumBones();
-	int64 Index = StartIndex;
+	// Make sure we don't write past the OutputBuffer. Six, because of two columns of the 3x3 rotation matrix.
+	const int32 AssetNumBones = Model->GetInputInfo()->GetNumBones();
+	checkfSlow((StartIndex + AssetNumBones * 6) <= OutputBufferSize, TEXT("Writing bones past the end of the input buffer."));
 
-	// Make sure we don't write past the OutputBuffer. (6 because of two columns of the 3x3 rotation matrix)
-	checkf((Index + AssetNumBones * 6) <= OutputBufferSize, TEXT("Writing bones past the end of the input buffer."));
+	// Write 6 floats to the buffer, for each bone.
+	UMLDeformerInputInfo::RotationToTwoVectorsAsSixFloats(TrainingBoneTransforms, OutputBuffer + StartIndex);
 
-	for (int32 BoneIndex = 0; BoneIndex < AssetNumBones; ++BoneIndex)
-	{
-		const FMatrix RotationMatrix = TrainingBoneTransforms[BoneIndex].GetRotation().ToMatrix();
-		const FVector X = RotationMatrix.GetColumn(0);
-		const FVector Y = RotationMatrix.GetColumn(1);	
-		OutputBuffer[Index++] = X.X;
-		OutputBuffer[Index++] = X.Y;
-		OutputBuffer[Index++] = X.Z;
-		OutputBuffer[Index++] = Y.X;
-		OutputBuffer[Index++] = Y.Y;
-		OutputBuffer[Index++] = Y.Z;
-	}
-
-	return Index;
+	// Return the new buffer offset, where we stopped writing.
+	return StartIndex + AssetNumBones * 6;
 }
 
 int64 UMLDeformerModelInstance::SetCurveValues(float* OutputBuffer, int64 OutputBufferSize, int64 StartIndex)
@@ -260,8 +247,6 @@ int64 UMLDeformerModelInstance::SetCurveValues(float* OutputBuffer, int64 Output
 
 int64 UMLDeformerModelInstance::SetNeuralNetworkInputValues(float* InputData, int64 NumInputFloats)
 {
-	check(SkeletalMeshComponent);
-
 	// Feed data to the network inputs.
 	int64 BufferOffset = 0;
 	if (Model->DoesSupportBones())
diff --git a/Engine/Plugins/Animation/MLDeformer/MLDeformerFramework/Source/MLDeformerFramework/Private/MLDeformerMorphModelInstance.cpp b/Engine/Plugins/Animation/MLDeformer/MLDeformerFramework/Source/MLDeformerFramework/Private/MLDeformerMorphModelInstance.cpp
--- a/Engine/Plugins/Animation/MLDeformer/MLDeformerFramework/Source/MLDeformerFramework/Private/MLDeformerMorphModelInstance.cpp
+++ b/Engine/Plugins/Animation/MLDeformer/MLDeformerFramework/Source/MLDeformerFramework/Private/MLDeformerMorphModelInstance.cpp
@@ -16,6 +16,11 @@ void UMLDeformerMorphModelInstance::Init(USkeletalMeshComponent* SkelMeshCompone
 	ExternalMorphSetID = NextFreeMorphSetID++;
 }
 
+bool UMLDeformerMorphModelInstance::IsValidForDataProvider() const
+{
+	return true;
+}
+
 void UMLDeformerMorphModelInstance::Release()
 {
 	// Try to unregister the morph target and morph target set.
@@ -93,20 +98,18 @@ void UMLDeformerMorphModelInstance::PostMLDeformerComponentInit()
 FExternalMorphSetWeights* UMLDeformerMorphModelInstance::FindWeightData(int32 LOD) const
 {
 	const UMLDeformerMorphModel* MorphModel = Cast<UMLDeformerMorphModel>(Model);
-	if (MorphModel == nullptr || !SkeletalMeshComponent)
-	{
-		return nullptr;
-	}
+	checkSlow(MorphModel);
 
-	// If we haven't got an external morph set registered yet, let's just return a nullptr.
-	if (!SkeletalMeshComponent->HasExternalMorphSet(LOD, ExternalMorphSetID))
+	// Check if our LOD index is valid first, as we might not have registered yet.
+	USkeletalMeshComponent* SkelMeshComponent = SkeletalMeshComponent.Get();
+	if (SkelMeshComponent == nullptr || !SkelMeshComponent->IsValidExternalMorphSetLODIndex(LOD))
 	{
 		return nullptr;
 	}
 
 	// Grab the weight data for this morph set.
 	// This could potentially fail if we are applying this deformer to the wrong skeletal mesh component.
-	return SkeletalMeshComponent->GetExternalMorphWeights(LOD).MorphSets.Find(ExternalMorphSetID);
+	return SkelMeshComponent->GetExternalMorphWeights(LOD).MorphSets.Find(ExternalMorphSetID);
 }
 
 void UMLDeformerMorphModelInstance::HandleZeroModelWeight()
@@ -124,7 +127,7 @@ void UMLDeformerMorphModelInstance::HandleZeroModelWeight()
 // Run the neural network, which calculates its outputs, which are the weights of our morph targets.
 void UMLDeformerMorphModelInstance::Execute(float ModelWeight)
 {
-	TRACE_CPUPROFILER_EVENT_SCOPE(UNeuralMorphModelInstance::Execute)
+	TRACE_CPUPROFILER_EVENT_SCOPE(UMLDeformerMorphModelInstance::Execute)
 
 	// Grab the weight data for this morph set.
 	// This could potentially fail if we are applying this deformer to the wrong skeletal mesh component.
diff --git a/Engine/Plugins/Animation/MLDeformer/MLDeformerFramework/Source/MLDeformerFramework/Private/MLDeformerObjectVersion.cpp b/Engine/Plugins/Animation/MLDeformer/MLDeformerFramework/Source/MLDeformerFramework/Private/MLDeformerObjectVersion.cpp
new file mode 100644
--- /dev/null
+++ b/Engine/Plugins/Animation/MLDeformer/MLDeformerFramework/Source/MLDeformerFramework/Private/MLDeformerObjectVersion.cpp
@@ -0,0 +1,13 @@
+// Copyright Epic Games, Inc. All Rights Reserved.
+
+#include "MLDeformerObjectVersion.h"
+#include "UObject/DevObjectVersion.h"
+
+namespace UE::MLDeformer
+{
+	// Unique MLDeformer Object version ID.
+	const FGuid FMLDeformerObjectVersion::GUID(0x23D9FD27, 0x4F6A0FDF, 0xEA055480, 0x7D4B3AB0);
+	
+	// Register the custom version with Core.
+	FDevVersionRegistration MLDeformerObjectVersionRegistration(FMLDeformerObjectVersion::GUID, FMLDeformerObjectVersion::LatestVersion, TEXT("Dev-MLDeformer"));
+}
diff --git a/Engine/Plugins/Animation/MLDeformer/MLDeformerFramework/Source/MLDeformerFramework/Public/MLDeformerAsset.h b/Engine/Plugins/Animation/MLDeformer/MLDeformerFramework/Source/MLDeformerFramework/Public/MLDeformerAsset.h
--- a/Engine/Plugins/Animation/MLDeformer/MLDeformerFramework/Source/MLDeformerFramework/Public/MLDeformerAsset.h
+++ b/Engine/Plugins/Animation/MLDeformer/MLDeformerFramework/Source/MLDeformerFramework/Public/MLDeformerAsset.h
@@ -19,6 +19,10 @@ class MLDEFORMERFRAMEWORK_API UMLDeformerAsset
 	GENERATED_BODY()
 
 public:
+	// UObject overrides.
+	virtual void Serialize(FArchive& Archive) override;
+	// ~END UObject overrides.
+
 	/**
 	 * Get the ML Deformer model that is being applied by this asset.
 	 * @return A pointer to the model.
diff --git a/Engine/Plugins/Animation/MLDeformer/MLDeformerFramework/Source/MLDeformerFramework/Public/MLDeformerInputInfo.h b/Engine/Plugins/Animation/MLDeformer/MLDeformerFramework/Source/MLDeformerFramework/Public/MLDeformerInputInfo.h
--- a/Engine/Plugins/Animation/MLDeformer/MLDeformerFramework/Source/MLDeformerFramework/Public/MLDeformerInputInfo.h
+++ b/Engine/Plugins/Animation/MLDeformer/MLDeformerFramework/Source/MLDeformerFramework/Public/MLDeformerInputInfo.h
@@ -25,18 +25,17 @@ public:
 	virtual ~UMLDeformerInputInfo() = default;
 
 	/**
-	 * This method is executed post loading.
-	 * On default it will initialize the FNames based on the bone and curve string based names.
+	 * This method is executed post loading of the ML Deformer asset.
 	 */
 	virtual void OnPostLoad();
 
 	/**
 	 * Check whether the current inputs are compatible with a given skeletal mesh.
-	 * @param SkeletalMesh The skeletal mesh to check compatibility with. This may not be a nullptr.
+	 * @param InSkeletalMesh The skeletal mesh to check compatibility with. This may not be a nullptr.
 	 * @return Returns true when we can safely apply the ML Deformer to a character using this skeletal mesh, otherwise false is returned.
 	 * @note Use GenerateCompatibilityErrorString to get the error report.
 	 */
-	virtual bool IsCompatible(USkeletalMesh* SkeletalMesh) const;
+	virtual bool IsCompatible(USkeletalMesh* InSkeletalMesh) const;
 
 	/**
 	 * Clear all contents.
@@ -46,17 +45,11 @@ public:
 
 	/**
 	 * Get the compatibility error report.
-	 * @param SkeletalMesh The skeletal mesh to check compatibility with.
+	 * @param InSkeletalMesh The skeletal mesh to check compatibility with.
 	 * @return Returns an empty string in case there are no compatibility issues, otherwise it contains a string that describes the issue(s).
 	 *         In case a nullptr is passed as SkeletalMesh parameter, an empty string is returned.
 	 */
-	virtual FString GenerateCompatibilityErrorString(USkeletalMesh* SkeletalMesh) const;
-
-	/** 
-	 * Update the FName arrays based on the name string arrays.
-	 * This is automatically called on PostLoad of the UMLDeformerAsset.
-	 */
-	virtual void UpdateFNames();
+	virtual FString GenerateCompatibilityErrorString(USkeletalMesh* InSkeletalMesh) const;
 
 	/**
 	 * Check whether we have any training inputs or not.
@@ -78,19 +71,12 @@ public:
 	 */
 	int32 GetNumBones() const;
 
-	/**
-	 * Get the bone name as a string, for a given bone we included during training.
-	 * @param Index The bone index, which is a number in range of [0..GetNumBones()-1].
-	 * @result The name of the bone.
-	 */
-	const FString& GetBoneNameString(int32 Index) const;
-
 	/**
 	 * Get the bone name as an FName, for a given bone we included during training.
 	 * @param Index The bone index, which is a number in range of [0..GetNumBones()-1].
 	 * @result The name of the bone.
 	 */
-	const FName GetBoneName(int32 Index) const;
+	FName GetBoneName(int32 Index) const;
 
 	/**
 	 * Get the number of curves that we trained on.
@@ -98,18 +84,6 @@ public:
 	 */
 	int32 GetNumCurves() const;
 
-	/**
-	 * Get the list of bone names that should be included as training inputs.
-	 * @return An array of strings, one for each bone name.
-	 */
-	TArray<FString>& GetBoneNameStrings();
-
-	/**
-	 * Get the list of curve names that should be included as training inputs.
-	 * @return An array of strings, one for each curve name.
-	 */
-	TArray<FString>& GetCurveNameStrings();
-
 	/**
 	 * Get the list of bone names that should be included as training inputs.
 	 * @return An array of FName objects, one for each bone.
@@ -122,19 +96,26 @@ public:
 	 */
 	TArray<FName>& GetCurveNames();
 
-	/**
-	 * Get the curve name as a string, for a given bone we included during training.
-	 * @param Index The curve index, which is a number in range of [0..GetNumCurves()-1].
-	 * @result The name of the curve.
-	 */
-	const FString& GetCurveNameString(int32 Index) const;
-
 	/**
 	 * Get the curve name as an FName, for a given curve we included during training.
 	 * @param Index The curve index, which is a number in range of [0..GetNumCurves()-1].
 	 * @result The name of the curve.
 	 */
-	const FName GetCurveName(int32 Index) const;
+	FName GetCurveName(int32 Index) const;
+
+#if WITH_EDITORONLY_DATA
+	UE_DEPRECATED(5.2, "Please use GetBoneNames instead.")
+	TArray<FString>& GetBoneNameStrings()						{ return BoneNameStrings_DEPRECATED; }
+
+	UE_DEPRECATED(5.2, "Please use GetCurveNames instead.")
+	TArray<FString>& GetCurveNameStrings()						{ return CurveNameStrings_DEPRECATED; }
+
+	UE_DEPRECATED(5.2, "Please use GetBoneName instead.")
+	const FString& GetBoneNameString(int32 Index) const			{ return BoneNameStrings_DEPRECATED[Index]; }
+
+	UE_DEPRECATED(5.2, "Please use GetCurveName instead.")
+	const FString& GetCurveNameString(int32 Index) const		{ return CurveNameStrings_DEPRECATED[Index]; }
+#endif
 
 	/** 
 	 * Extract the curve values for all curves we're interested in.
@@ -143,6 +124,19 @@ public:
 	 */
 	virtual void ExtractCurveValues(USkeletalMeshComponent* SkelMeshComponent, TArray<float>& OutValues) const;
 
+	/**
+	 * Batch convert a set of transform rotation quaternions into two column vectors.
+	 * 
+	 */
+	static void RotationToTwoVectorsAsSixFloats(TArrayView<FTransform> Transforms, float* OutputBuffer);
+
+	/**
+	 * Convert a rotation quaternion to two basis vectors, each represented as 3 floats, so 6 floats in total.
+	 * @param Rotation The rotation quaternion to convert into the six float values.
+	 * @param SixFloatsOutputBuffer The float buffer that we will write the 6 floats to. Make sure this buffer is at least 6 floats large.
+	 */
+	static void RotationToTwoVectorsAsSixFloats(const FQuat& Rotation, float* SixFloatsOutputBuffer);
+
 	/**
 	 * Extract bone space rotations, as a float array.
 	 * The number of output rotations are NumBones * 6, where the array contains a set two columns of a 3x3 rotation matrix values.
@@ -179,17 +173,47 @@ public:
 	 */
 	void SetNumTargetVertices(int32 NumVerts);
 
+	/** 
+	 * Get the path to the skeletal mesh that this deformer was trained on.
+	 * @return The path to the skeletal mesh we used during training.
+	 */
+	const FSoftObjectPath& GetSkeletalMesh() const;
+
+	/**
+	 * Set skeletal mesh that we trained on.
+	 * @param InSkeletalMesh A pointer to the skeletal mesh.
+	 */
+	void SetSkeletalMesh(USkeletalMesh* InSkeletalMesh);
+
+	UE_DEPRECATED(5.2, "This method will be removed soon. Please only use the BoneNames and CurveNames and not their string versions.")
+	void UpdateNameStrings();
+
+	UE_DEPRECATED(5.2, "This method will be removed soon. It shouldn't be needed anymore.")
+	void UpdateFNames();
+
 protected:
+	/**
+	 * The path to the skeletal mesh that this model was trained on.
+	 */
+	UPROPERTY()	
+	FSoftObjectPath SkeletalMesh;
+
+#if WITH_EDITORONLY_DATA
+	/** The list of bone names, but as string. This is deprecated since UE 5.2. */
+	UPROPERTY(meta = (DeprecatedProperty, DeprecationMessage = "Use BoneNames instead."))
+	TArray<FString> BoneNameStrings_DEPRECATED;
+
+	/** The list of curve names, but as string. This is deprecated since UE 5.2. */
+	UPROPERTY(meta = (DeprecatedProperty, DeprecationMessage = "Use CurveNames instead."))
+	TArray<FString> CurveNameStrings_DEPRECATED;
+#endif
+
 	/** 
 	 * The name of each bone. The inputs to the network are in the order of this array.
 	 * So if the array contains ["Root", "Child1", "Child2"] then the first bone transforms that we 
 	 * input to the neural network is the transform for "Root", followed by "Child1", followed by "Child2".
 	 */
 	UPROPERTY()
-	TArray<FString> BoneNameStrings;
-
-	/** The same as the BoneNames member, but stored as pre-created FName objects. These are not serialized. */
-	UPROPERTY(Transient)
 	TArray<FName> BoneNames;
 
 	/**
@@ -198,10 +222,6 @@ protected:
 	 * input to the neural network is the one for "Smile", followed by "LeftEyeClosed", followed by "RightEyeClosed".
 	 */
 	UPROPERTY()
-	TArray<FString> CurveNameStrings;
-
-	/** The same as the CurveNames member, but stored as pre-created FName objects. These are not serialized. */
-	UPROPERTY(Transient)
 	TArray<FName> CurveNames;
 
 	/** Number of imported base mesh vertices, so not render vertices. */
diff --git a/Engine/Plugins/Animation/MLDeformer/MLDeformerFramework/Source/MLDeformerFramework/Public/MLDeformerMorphModelInstance.h b/Engine/Plugins/Animation/MLDeformer/MLDeformerFramework/Source/MLDeformerFramework/Public/MLDeformerMorphModelInstance.h
--- a/Engine/Plugins/Animation/MLDeformer/MLDeformerFramework/Source/MLDeformerFramework/Public/MLDeformerMorphModelInstance.h
+++ b/Engine/Plugins/Animation/MLDeformer/MLDeformerFramework/Source/MLDeformerFramework/Public/MLDeformerMorphModelInstance.h
@@ -29,6 +29,7 @@ public:
 	virtual void Release() override;
 	virtual void Execute(float ModelWeight) override;
 	virtual void HandleZeroModelWeight() override;
+	virtual bool IsValidForDataProvider() const override;
 	// ~END UMLDeformerModelInstance overrides.
 
 	int32 GetExternalMorphSetID() const { return ExternalMorphSetID; }
diff --git a/Engine/Plugins/Animation/MLDeformer/MLDeformerFramework/Source/MLDeformerFramework/Public/MLDeformerObjectVersion.h b/Engine/Plugins/Animation/MLDeformer/MLDeformerFramework/Source/MLDeformerFramework/Public/MLDeformerObjectVersion.h
new file mode 100644
--- /dev/null
+++ b/Engine/Plugins/Animation/MLDeformer/MLDeformerFramework/Source/MLDeformerFramework/Public/MLDeformerObjectVersion.h
@@ -0,0 +1,27 @@
+// Copyright Epic Games, Inc. All Rights Reserved.
+#pragma once
+
+#include "CoreTypes.h"
+#include "Misc/Guid.h"
+
+namespace UE::MLDeformer
+{
+	// Custom serialization version for backwards compatibility during de-serialization
+	struct MLDEFORMERFRAMEWORK_API FMLDeformerObjectVersion
+	{
+		FMLDeformerObjectVersion() = delete;
+	
+		enum Type
+		{
+			// Before any version changes were made.
+			BeforeCustomVersionWasAdded,
+
+			// ----- New versions can be added above this line -----
+			VersionPlusOne,
+			LatestVersion = VersionPlusOne - 1
+		};
+
+		// The GUID for this custom version number.
+		const static FGuid GUID;
+	};
+}	// namespace UE::MLDeformer
diff --git a/Engine/Plugins/Animation/MLDeformer/MLDeformerFramework/Source/MLDeformerFrameworkEditor/Private/MLDeformerEditorModel.cpp b/Engine/Plugins/Animation/MLDeformer/MLDeformerFramework/Source/MLDeformerFrameworkEditor/Private/MLDeformerEditorModel.cpp
--- a/Engine/Plugins/Animation/MLDeformer/MLDeformerFramework/Source/MLDeformerFrameworkEditor/Private/MLDeformerEditorModel.cpp
+++ b/Engine/Plugins/Animation/MLDeformer/MLDeformerFramework/Source/MLDeformerFrameworkEditor/Private/MLDeformerEditorModel.cpp
@@ -966,14 +966,13 @@ namespace UE::MLDeformer
 	{
 		InputInfo->Reset();
 
-		TArray<FString>& BoneNameStrings = InputInfo->GetBoneNameStrings();
-		TArray<FString>& CurveNameStrings = InputInfo->GetCurveNameStrings();
 		TArray<FName>& BoneNames = InputInfo->GetBoneNames();
 		TArray<FName>& CurveNames = InputInfo->GetCurveNames();
+
 		USkeletalMesh* SkeletalMesh = Model->GetSkeletalMesh();
+		InputInfo->SetSkeletalMesh(SkeletalMesh);
 
 		BoneNames.Reset();
-		BoneNameStrings.Reset();
 		CurveNames.Reset();
 
 		InputInfo->SetNumBaseVertices(Model->GetNumBaseMeshVerts());
@@ -992,11 +991,10 @@ namespace UE::MLDeformer
 			{
 				// Grab all bone names.
 				const int32 NumBones = RefSkeleton.GetNum();
-				BoneNameStrings.Reserve(NumBones);
+				BoneNames.Reserve(NumBones);
 				for (int32 Index = 0; Index < NumBones; ++Index)
 				{
 					const FName BoneName = RefSkeleton.GetBoneName(Index);
-					BoneNameStrings.Add(BoneName.ToString());
 					BoneNames.Add(BoneName);
 				}
 			}
@@ -1013,7 +1011,6 @@ namespace UE::MLDeformer
 							continue;
 						}
 
-						BoneNameStrings.Add(BoneName.ToString());
 						BoneNames.Add(BoneName);
 					}
 				}
@@ -1031,11 +1028,6 @@ namespace UE::MLDeformer
 				if (Model->GetCurveIncludeList().IsEmpty())
 				{
 					SmartNameMapping->FillNameArray(CurveNames);
-					CurveNameStrings.Reserve(CurveNames.Num());
-					for (const FName Name : CurveNames)
-					{
-						CurveNameStrings.Add(Name.ToString());
-					}
 				}
 				else // A list of curve names was provided.
 				{
@@ -1050,13 +1042,16 @@ namespace UE::MLDeformer
 								continue;
 							}
 
-							CurveNameStrings.Add(CurveName.ToString());
 							CurveNames.Add(CurveName);
 						}
 					}
 				}
 			}
 		}
+
+		PRAGMA_DISABLE_DEPRECATION_WARNINGS
+		InputInfo->UpdateNameStrings();
+		PRAGMA_ENABLE_DEPRECATION_WARNINGS
 	}
 
 	void FMLDeformerEditorModel::InitBoneIncludeListToAnimatedBonesOnly()
@@ -1616,9 +1611,9 @@ namespace UE::MLDeformer
 
 		// Initialize an engine morph target for each model morph target.
 		UE_LOG(LogMLDeformer, Display, TEXT("Initializing %d engine morph targets of %d vertices each"), NumMorphTargets, Deltas.Num() / NumMorphTargets);
-		for (int32 BlendShapeIndex = 0; BlendShapeIndex < NumMorphTargets; ++BlendShapeIndex)
+		for (int32 MorphTargetIndex = 0; MorphTargetIndex < NumMorphTargets; ++MorphTargetIndex)
 		{
-			const FName MorphName = *FString::Printf(TEXT("%s%.3d"), *NamePrefix, BlendShapeIndex);
+			const FName MorphName = *FString::Printf(TEXT("%s%.3d"), *NamePrefix, MorphTargetIndex);
 			UMorphTarget* MorphTarget = NewObject<UMorphTarget>(SkelMesh, MorphName);
 			MorphTarget->BaseSkelMesh = SkelMesh;
 			OutMorphTargets.Add(MorphTarget);
@@ -1650,7 +1645,7 @@ namespace UE::MLDeformer
 				const int32 ImportedVertexNumber = VertexMap[VertexIndex];
 				if (ImportedVertexNumber != INDEX_NONE)
 				{
-					const FVector3f Delta = Deltas[ImportedVertexNumber + BlendShapeIndex * NumBaseMeshVerts];
+					const FVector3f Delta = Deltas[ImportedVertexNumber + MorphTargetIndex * NumBaseMeshVerts];
 					if (Delta.Length() > DeltaThreshold)
 					{
 						MorphLODModel.Vertices.AddDefaulted();
diff --git a/Engine/Plugins/Animation/MLDeformer/MLDeformerFramework/Source/MLDeformerFrameworkEditor/Private/MLDeformerEditorToolkit.cpp b/Engine/Plugins/Animation/MLDeformer/MLDeformerFramework/Source/MLDeformerFrameworkEditor/Private/MLDeformerEditorToolkit.cpp
--- a/Engine/Plugins/Animation/MLDeformer/MLDeformerFramework/Source/MLDeformerFrameworkEditor/Private/MLDeformerEditorToolkit.cpp
+++ b/Engine/Plugins/Animation/MLDeformer/MLDeformerFramework/Source/MLDeformerFrameworkEditor/Private/MLDeformerEditorToolkit.cpp
@@ -548,7 +548,8 @@ namespace UE::MLDeformer
 			break;
 
 			// Unhandled error codes.
-			default: check(false);
+			default:
+				checkf(false, TEXT("Unknown error code"));
 		}
 
 		// Show a message window.
diff --git a/Engine/Plugins/Animation/MLDeformer/MLDeformerFramework/Source/MLDeformerFrameworkEditor/Private/MLDeformerModelDetails.cpp b/Engine/Plugins/Animation/MLDeformer/MLDeformerFramework/Source/MLDeformerFrameworkEditor/Private/MLDeformerModelDetails.cpp
--- a/Engine/Plugins/Animation/MLDeformer/MLDeformerFramework/Source/MLDeformerFrameworkEditor/Private/MLDeformerModelDetails.cpp
+++ b/Engine/Plugins/Animation/MLDeformer/MLDeformerFramework/Source/MLDeformerFrameworkEditor/Private/MLDeformerModelDetails.cpp
@@ -311,9 +311,8 @@ namespace UE::MLDeformer
 
 		// Show a warning when no neural network has been set.
 		{		
-			UNeuralNetwork* NeuralNetwork = Model->GetNeuralNetwork();
 			FDetailWidgetRow& NeuralNetErrorRow = TrainingSettingsCategoryBuilder->AddCustomRow(FText::FromString("NeuralNetError"))
-				.Visibility((NeuralNetwork == nullptr) ? EVisibility::Visible : EVisibility::Collapsed)
+				.Visibility(!EditorModel->IsTrained() ? EVisibility::Visible : EVisibility::Collapsed)
 				.WholeRowContent()
 				[
 					SNew(SBox)
@@ -321,12 +320,12 @@ namespace UE::MLDeformer
 					[
 						SNew(SWarningOrErrorBox)
 						.MessageStyle(EMessageStyle::Warning)
-						.Message(FText::FromString("Model still needs to be trained."))
+						.Message(LOCTEXT("NeedsTraining", "Model still needs to be trained."))
 					]
 				];
 
 			// Check if our network is compatible with the skeletal mesh.
-			if (Model->GetSkeletalMesh() && NeuralNetwork)
+			if (Model->GetSkeletalMesh() && EditorModel->IsTrained())
 			{
 				FDetailWidgetRow& NeuralNetIncompatibleErrorRow = TrainingSettingsCategoryBuilder->AddCustomRow(FText::FromString("NeuralNetIncompatibleError"))
 					.Visibility(!Model->GetInputInfo()->IsCompatible(Model->GetSkeletalMesh()) ? EVisibility::Visible : EVisibility::Collapsed)
@@ -337,11 +336,13 @@ namespace UE::MLDeformer
 						[
 							SNew(SWarningOrErrorBox)
 							.MessageStyle(EMessageStyle::Error)
-							.Message(FText::FromString("Trained neural network is incompatible with selected SkeletalMesh."))
+							.Message(LOCTEXT("TrainingIncompatibleWithSkelMesh", "Trained neural network is incompatible with selected SkeletalMesh."))
 						]
 					];
 			}
 		}
+
+		AddTrainingSettingsErrors();
 	}
 
 	bool FMLDeformerModelDetails::FilterAnimSequences(const FAssetData& AssetData, USkeleton* Skeleton)
diff --git a/Engine/Plugins/Animation/MLDeformer/MLDeformerFramework/Source/MLDeformerFrameworkEditor/Public/MLDeformerModelDetails.h b/Engine/Plugins/Animation/MLDeformer/MLDeformerFramework/Source/MLDeformerFrameworkEditor/Public/MLDeformerModelDetails.h
--- a/Engine/Plugins/Animation/MLDeformer/MLDeformerFramework/Source/MLDeformerFrameworkEditor/Public/MLDeformerModelDetails.h
+++ b/Engine/Plugins/Animation/MLDeformer/MLDeformerFramework/Source/MLDeformerFrameworkEditor/Public/MLDeformerModelDetails.h
@@ -68,6 +68,9 @@ namespace UE::MLDeformer
 		/** Add additional training input errors. */
 		virtual void AddTrainingInputErrors() {}
 
+		/** Add training settings errors. */
+		virtual void AddTrainingSettingsErrors() {}
+
 		/** Is the "include bones" checkbox visible? */
 		virtual bool IsBonesFlagVisible() const;
 
diff --git a/Engine/Plugins/Animation/MLDeformer/NearestNeighborModel/Source/NearestNeighborModel/Private/NearestNeighborModelInputInfo.cpp b/Engine/Plugins/Animation/MLDeformer/NearestNeighborModel/Source/NearestNeighborModel/Private/NearestNeighborModelInputInfo.cpp
--- a/Engine/Plugins/Animation/MLDeformer/NearestNeighborModel/Source/NearestNeighborModel/Private/NearestNeighborModelInputInfo.cpp
+++ b/Engine/Plugins/Animation/MLDeformer/NearestNeighborModel/Source/NearestNeighborModel/Private/NearestNeighborModelInputInfo.cpp
@@ -59,8 +59,8 @@ void UNearestNeighborModelInputInfo::ExtractBoneRotations(USkeletalMeshComponent
 int32 UNearestNeighborModelInputInfo::CalcNumNeuralNetInputs() const
 {
     return 
-        BoneNameStrings.Num() * 3 + // Three floats per bone.
-        CurveNameStrings.Num();     // One float per curve.
+        BoneNames.Num() * 3 + // Three floats per bone.
+        CurveNames.Num();     // One float per curve.
 }
 
 void UNearestNeighborModelInputInfo::InitRefBoneRotations(USkeletalMesh* SkelMesh)
diff --git a/Engine/Plugins/Animation/MLDeformer/NeuralMorphModel/Source/NeuralMorphModel/Private/NeuralMorphModel.cpp b/Engine/Plugins/Animation/MLDeformer/NeuralMorphModel/Source/NeuralMorphModel/Private/NeuralMorphModel.cpp
--- a/Engine/Plugins/Animation/MLDeformer/NeuralMorphModel/Source/NeuralMorphModel/Private/NeuralMorphModel.cpp
+++ b/Engine/Plugins/Animation/MLDeformer/NeuralMorphModel/Source/NeuralMorphModel/Private/NeuralMorphModel.cpp
@@ -1,6 +1,10 @@
 // Copyright Epic Games, Inc. All Rights Reserved.
 #include "NeuralMorphModel.h"
 #include "NeuralMorphModelVizSettings.h"
+#include "NeuralMorphModelInstance.h"
+#include "MLDeformerAsset.h"
+#include "MLDeformerComponent.h"
+#include "UObject/UObjectGlobals.h"
 
 #define LOCTEXT_NAMESPACE "NeuralMorphModel"
 
@@ -10,6 +14,12 @@ namespace UE::NeuralMorphModel
 	class NEURALMORPHMODEL_API FNeuralMorphModelModule
 		: public IModuleInterface
 	{
+		void StartupModule() override
+		{
+			#if NEURALMORPHMODEL_FORCE_USE_NNI
+				UE_LOG(LogNeuralMorphModel, Warning, TEXT("Running neural morph model with NNI. The faster custom inference code path will be disabled."));
+			#endif
+		}
 	};
 }
 IMPLEMENT_MODULE(UE::NeuralMorphModel::FNeuralMorphModelModule, NeuralMorphModel)
@@ -30,4 +40,46 @@ UNeuralMorphModel::UNeuralMorphModel(const FObjectInitializer& ObjectInitializer
 #endif
 }
 
+UMLDeformerModelInstance* UNeuralMorphModel::CreateModelInstance(UMLDeformerComponent* Component)
+{
+	return NewObject<UNeuralMorphModelInstance>(Component);
+}
+
+void UNeuralMorphModel::Serialize(FArchive& Archive)
+{
+#if !NEURALMORPHMODEL_FORCE_USE_NNI
+	if (Archive.IsSaving() || Archive.IsCooking())
+	{
+		UNeuralNetwork* NNINeuralNetwork = GetNeuralNetwork();
+
+		// Show a warning when we are not using custom inference yet on this model.
+		if (NeuralMorphNetwork == nullptr && NNINeuralNetwork != nullptr)
+		{
+			UE_LOG(LogNeuralMorphModel, Warning, TEXT("Neural Morph Model in MLD asset '%s' should be retrained to get higher performance by taking advantage of custom inference."), *GetDeformerAsset()->GetName());
+		}
+
+		// If we have a custom inference network, make sure we don't save out the NNI network.
+		if (NeuralMorphNetwork != nullptr && NNINeuralNetwork != nullptr)
+		{
+			SetNeuralNetwork(nullptr);
+		}
+	}
+#endif
+
+	Super::Serialize(Archive);
+}
+
+void UNeuralMorphModel::PostLoad()
+{
+	Super::PostLoad();
+
+#if !NEURALMORPHMODEL_FORCE_USE_NNI
+	// Show a warning when we are not using custom inference yet on this model.
+	if (NeuralMorphNetwork == nullptr && GetNeuralNetwork() != nullptr)
+	{
+		UE_LOG(LogNeuralMorphModel, Warning, TEXT("Neural Morph Model in MLD asset '%s' should be retrained to get higher performance by taking advantage of custom inference."), *GetDeformerAsset()->GetName());
+	}
+#endif
+}
+
 #undef LOCTEXT_NAMESPACE
diff --git a/Engine/Plugins/Animation/MLDeformer/NeuralMorphModel/Source/NeuralMorphModel/Private/NeuralMorphModelInstance.cpp b/Engine/Plugins/Animation/MLDeformer/NeuralMorphModel/Source/NeuralMorphModel/Private/NeuralMorphModelInstance.cpp
new file mode 100644
--- /dev/null
+++ b/Engine/Plugins/Animation/MLDeformer/NeuralMorphModel/Source/NeuralMorphModel/Private/NeuralMorphModelInstance.cpp
@@ -0,0 +1,120 @@
+// Copyright Epic Games, Inc. All Rights Reserved.
+
+#include "NeuralMorphModelInstance.h"
+#include "NeuralMorphModel.h"
+#include "NeuralMorphNetwork.h"
+#include "MLDeformerAsset.h"
+#include "Components/ExternalMorphSet.h"
+#include "Components/SkeletalMeshComponent.h"
+
+
+void UNeuralMorphModelInstance::Init(USkeletalMeshComponent* SkelMeshComponent)
+{
+	Super::Init(SkelMeshComponent);
+
+#if !NEURALMORPHMODEL_FORCE_USE_NNI
+	UNeuralMorphModel* MorphModel = Cast<UNeuralMorphModel>(Model);
+	UNeuralMorphNetwork* MorphNetwork = MorphModel->GetNeuralMorphNetwork();
+	NetworkInstance = MorphNetwork ? MorphNetwork->CreateInstance() : nullptr;
+#endif
+}
+
+bool UNeuralMorphModelInstance::SetupInputs()
+{
+	TRACE_CPUPROFILER_EVENT_SCOPE(UNeuralMorphModelInstance::SetupInputs)
+
+#if NEURALMORPHMODEL_FORCE_USE_NNI
+	return Super::SetupInputs();
+#else
+	// If we have no neural morph network, fall back to the default NNI path.
+	UNeuralMorphModel* MorphModel = Cast<UNeuralMorphModel>(Model);
+	UNeuralMorphNetwork* MorphNetwork = MorphModel->GetNeuralMorphNetwork();
+	if (MorphNetwork == nullptr)
+	{		
+		return Super::SetupInputs();
+	}
+
+	// Some safety checks.
+	if (SkeletalMeshComponent == nullptr ||
+		SkeletalMeshComponent->GetSkeletalMeshAsset() == nullptr ||
+		!bIsCompatible)
+	{
+		return false;
+	}
+
+	// If the neural network expects a different number of inputs, do nothing.
+	const int64 NumNeuralNetInputs = MorphNetwork->GetNumInputs();
+	const int64 NumDeformerAssetInputs = Model->GetInputInfo()->CalcNumNeuralNetInputs();
+	if (NumNeuralNetInputs != NumDeformerAssetInputs)
+	{
+		return false;
+	}
+
+	// Update and write the input values directly into the input tensor.
+	float* InputDataPointer = NetworkInstance->GetInputs().GetData();
+	const int64 NumFloatsWritten = SetNeuralNetworkInputValues(InputDataPointer, NumNeuralNetInputs);
+	check(NumFloatsWritten == NumNeuralNetInputs);
+
+	return true;
+#endif
+}
+
+void UNeuralMorphModelInstance::Execute(float ModelWeight)
+{
+	TRACE_CPUPROFILER_EVENT_SCOPE(UNeuralMorphModelInstance::Execute)
+
+#if NEURALMORPHMODEL_FORCE_USE_NNI
+	Super::Execute(ModelWeight);
+	return;
+#else
+	UNeuralMorphModel* MorphModel = Cast<UNeuralMorphModel>(Model);
+	UNeuralMorphNetwork* MorphNetwork = MorphModel->GetNeuralMorphNetwork();
+	if (MorphNetwork == nullptr)
+	{
+		Super::Execute(ModelWeight);
+		return;
+	}
+
+	// Grab the weight data for this morph set.
+	// This could potentially fail if we are applying this deformer to the wrong skeletal mesh component.
+	const int LOD = 0;	// For now we only support LOD 0, as we can't setup an ML Deformer per LOD yet.
+	FExternalMorphSetWeights* WeightData = FindWeightData(LOD);
+	if (WeightData == nullptr)
+	{
+		return;
+	}
+
+	if (!MorphNetwork->IsEmpty())
+	{
+		// Perform inference on the neural network, this updates its output values.
+		NetworkInstance->Run();
+
+		// If our model is active, we want to run the neural network and update the morph weights
+		// with the values that the neural net calculated for us.
+		// Get the network output values, read the values and use them as morph target weights inside the skeletal mesh component.
+		const TArrayView<const float> NetworkOutputs = NetworkInstance->GetOutputs();
+		const int32 NumNetworkWeights = NetworkOutputs.Num();
+		const int32 NumMorphTargets = WeightData->Weights.Num();
+
+		// If we have the expected amount of weights.
+		// +1 because we always have an extra morph target that represents the means, with fixed weight of 1.
+		// Therefore, the neural network output will contain one less float than the number of morph targets in our morph set.
+		if (NumMorphTargets == NumNetworkWeights + 1)
+		{
+			// Set the first morph target, which represents the means, to a weight of 1.0, as it always needs to be fully active.
+			WeightData->Weights[0] = 1.0f * ModelWeight;
+
+			// Update all generated morph target weights with the values calculated by our neural network.
+			for (int32 MorphIndex = 0; MorphIndex < NumNetworkWeights; ++MorphIndex)
+			{
+				const float OutputTensorValue = NetworkOutputs[MorphIndex];
+				WeightData->Weights[MorphIndex + 1] = OutputTensorValue * ModelWeight;
+			}
+			return;
+		}
+	}
+
+	// Set the weights to zero if we have no valid network.
+	WeightData->ZeroWeights();
+#endif
+}
diff --git a/Engine/Plugins/Animation/MLDeformer/NeuralMorphModel/Source/NeuralMorphModel/Private/NeuralMorphNetwork.cpp b/Engine/Plugins/Animation/MLDeformer/NeuralMorphModel/Source/NeuralMorphModel/Private/NeuralMorphNetwork.cpp
new file mode 100644
--- /dev/null
+++ b/Engine/Plugins/Animation/MLDeformer/NeuralMorphModel/Source/NeuralMorphModel/Private/NeuralMorphNetwork.cpp
@@ -0,0 +1,491 @@
+// Copyright Epic Games, Inc. All Rights Reserved.
+
+#include "NeuralMorphNetwork.h"
+#include "NeuralMorphModel.h"
+#include "UObject/SoftObjectPath.h"
+#include "HAL/FileManagerGeneric.h"
+#include "Math/UnrealMathUtility.h"
+
+#if NEURALMORPHMODEL_USE_ISPC
+	#include "NeuralMorphNetwork.ispc.generated.h"
+#endif
+
+void UNeuralMorphNetwork::Empty()
+{
+	InputMeans.Empty();
+	InputStd.Empty();
+	Layers.Empty();
+	Mode = ENeuralMorphMode::Global;
+	NumMorphsPerBone = 0;
+	NumBones = 0;
+	NumCurves = 0;
+}
+
+bool UNeuralMorphNetwork::IsEmpty() const
+{
+	return GetNumInputs() == 0;
+}
+
+int32 UNeuralMorphNetwork::GetNumInputs() const
+{
+	return !Layers.IsEmpty() ? Layers[0]->NumInputs * Layers[0]->Depth : 0;
+}
+
+int32 UNeuralMorphNetwork::GetNumOutputs() const
+{
+	return !Layers.IsEmpty() ? Layers[Layers.Num()-1]->NumOutputs * Layers[Layers.Num()-1]->Depth : 0;
+}
+
+bool UNeuralMorphNetwork::Load(const FString& Filename)
+{
+	UE_LOG(LogNeuralMorphModel, Display, TEXT("Loading Neural Morph Network from file '%s'"), *Filename);
+	Empty();
+
+	// Create the file reader.
+	TUniquePtr<FArchive> FileReader(IFileManager::Get().CreateFileReader(*Filename));
+	if (!FileReader.IsValid())
+	{
+		UE_LOG(LogNeuralMorphModel, Error, TEXT("Failed to create file reader to read file '%s'!"), *Filename);
+		Empty();
+		return false;
+	}
+
+	// Read the FOURCC, to identify the file type.
+	char FOURCC[4] {' ', ' ', ' ', ' '};
+	FileReader->Serialize(FOURCC, 4);
+	if (FileReader->IsError())
+	{
+		FileReader->Close();
+		UE_LOG(LogNeuralMorphModel, Error, TEXT("Failed to read the FOURCC!"));
+		Empty();
+		return false;
+	}
+
+	if (FOURCC[0] != 'N' || FOURCC[1] != 'M' || FOURCC[2] != 'M' || FOURCC[3] != 'N')	// NMMN (Neural Morph Model Network)
+	{
+		FileReader->Close();
+		UE_LOG(LogNeuralMorphModel, Error, TEXT("The file is not a valid valid neural morph network file type!"));
+		Empty();
+		return false;
+	}
+
+	// Load the version number.
+	int32 Version = -1;
+	FileReader->Serialize(&Version, sizeof(int32));
+	if (FileReader->IsError())
+	{
+		FileReader->Close();
+		UE_LOG(LogNeuralMorphModel, Error, TEXT("Failed to load version"));
+		Empty();
+		return false;
+	}
+
+	if (Version != 1)
+	{
+		UE_LOG(LogNeuralMorphModel, Error, TEXT("The Neural Morph Network file '%s' is of an unknown version (Version=%d)!"), *Filename, Version);
+		FileReader->Close();
+		Empty();
+		return false;
+	}
+
+	// Load the header with info.
+	struct FInfoHeader
+	{
+		int32 Mode = -1;					// 0=Local, 1=Global
+		int32 NumInputs = -1;				// How many float inputs in the input layer?
+		int32 NumHiddenLayers = -1;			// How many hidden layers in the network (layers excluding input and output layer)?
+		int32 NumUnitsPerHiddenLayer = -1;	// How many neurons for each hidden layer?
+		int32 NumOutputs = -1;				// The number of units in the output layer.
+		int32 NumMorphsPerBone = -1;		// The number of morph targets per bone, if set Mode == 0. Otherwise ignored.
+		int32 NumBones = -1;				// The number of bones used as input.
+		int32 NumCurves = -1;				// The number of curves used as input.
+	};
+	FInfoHeader Info;
+	FileReader->Serialize(&Info, sizeof(FInfoHeader));
+	if (FileReader->IsError())
+	{
+		UE_LOG(LogNeuralMorphModel, Error, TEXT("Failed to load info header!"));
+		FileReader->Close();
+		Empty();
+		return false;
+	}
+
+	if (Info.Mode == -1 || Info.NumInputs == -1 || Info.NumHiddenLayers == -1 || Info.NumUnitsPerHiddenLayer == -1 || Info.NumOutputs == -1 || Info.NumBones == -1 || Info.NumCurves == -1)
+	{
+		UE_LOG(LogNeuralMorphModel, Error, TEXT("Failed to read info header for Neural Morph Network in file '%s'!"), *Filename);
+		FileReader->Close();
+		Empty();
+		return false;
+	}
+
+	Mode = (Info.Mode == 0) ? ENeuralMorphMode::Local : ENeuralMorphMode::Global;
+	NumMorphsPerBone = Info.NumMorphsPerBone;
+	NumBones = Info.NumBones;
+	NumCurves = Info.NumCurves;
+
+	// Read the input standard deviation and means.
+	InputMeans.SetNumZeroed(Info.NumInputs);
+	InputStd.SetNumZeroed(Info.NumInputs);
+	FileReader->Serialize(InputStd.GetData(), Info.NumInputs * sizeof(float));
+	if (FileReader->IsError())
+	{
+		UE_LOG(LogNeuralMorphModel, Error, TEXT("Failed to load input standard deviations!"));
+		FileReader->Close();
+		Empty();
+		return false;
+	}
+
+	FileReader->Serialize(InputMeans.GetData(), Info.NumInputs * sizeof(float));
+	if (FileReader->IsError())
+	{
+		UE_LOG(LogNeuralMorphModel, Error, TEXT("Failed to load input means!"));
+		FileReader->Close();
+		Empty();
+		return false;
+	}
+
+	struct FLayerHeader
+	{
+		int32 NumInputs = -1;
+		int32 NumOutputs = -1;
+		int32 NumWeights = -1;
+		int32 NumBiases = -1;
+	};
+
+	// Load the weights and biases.
+	for (int32 LayerIndex = 0; LayerIndex < Info.NumHiddenLayers + 1; ++LayerIndex)
+	{
+		UNeuralMorphNetworkLayer* WeightData = NewObject<UNeuralMorphNetworkLayer>(this);
+
+		FLayerHeader LayerHeader;
+		FileReader->Serialize(&LayerHeader, sizeof(FLayerHeader));
+		if (FileReader->IsError())
+		{
+			UE_LOG(LogNeuralMorphModel, Error, TEXT("Failed to load layer header!"));
+			FileReader->Close();
+			Empty();
+			return false;
+		}
+
+		WeightData->NumInputs = LayerHeader.NumInputs;
+		WeightData->NumOutputs = LayerHeader.NumOutputs;
+		WeightData->Depth = LayerHeader.NumWeights / (LayerHeader.NumInputs * LayerHeader.NumOutputs);
+
+		WeightData->Weights.AddZeroed(LayerHeader.NumWeights);
+		FileReader->Serialize(WeightData->Weights.GetData(), LayerHeader.NumWeights * sizeof(float));
+		if (FileReader->IsError())
+		{
+			UE_LOG(LogNeuralMorphModel, Error, TEXT("Failed to load layer weights!"));
+			FileReader->Close();
+			Empty();
+			return false;
+		}
+
+		WeightData->Biases.AddZeroed(LayerHeader.NumBiases);
+		FileReader->Serialize(WeightData->Biases.GetData(), LayerHeader.NumBiases * sizeof(float));
+		if (FileReader->IsError())
+		{
+			UE_LOG(LogNeuralMorphModel, Error, TEXT("Failed to load layer biases!"));
+			FileReader->Close();
+			Empty();
+			return false;
+		}
+
+		Layers.Add(WeightData);
+
+		UE_LOG(LogNeuralMorphModel, Verbose, TEXT("Network Layer %d --> NumWeights=%d (%dx%dx%d)   NumBiases=%d"), LayerIndex, WeightData->Weights.Num(), WeightData->NumInputs, WeightData->NumOutputs, WeightData->Depth, WeightData->Biases.Num());
+	}
+
+	// Clean up and return the result.
+	const bool bSuccess = FileReader->Close();
+	if (bSuccess)
+	{
+		UE_LOG(LogNeuralMorphModel, Display, TEXT("Successfullly loaded neural morph network from file '%s'"), *Filename);
+	}
+	else
+	{
+		UE_LOG(LogNeuralMorphModel, Error, TEXT("Failed to close file reader."));
+		Empty();
+	}
+
+	return bSuccess;
+}
+
+UNeuralMorphNetworkInstance* UNeuralMorphNetwork::CreateInstance()
+{
+	UNeuralMorphNetworkInstance* Instance = NewObject<UNeuralMorphNetworkInstance>(this);
+	Instance->Init(this);
+	return Instance;
+}
+
+int32 UNeuralMorphNetwork::GetNumBones() const
+{
+	return NumBones;
+}
+
+int32 UNeuralMorphNetwork::GetNumCurves() const
+{
+	return NumCurves;
+}
+
+int32 UNeuralMorphNetwork::GetNumMorphsPerBone() const
+{
+	return NumMorphsPerBone;
+}
+
+ENeuralMorphMode UNeuralMorphNetwork::GetMode() const
+{
+	return Mode;
+}
+
+const TArrayView<const float> UNeuralMorphNetwork::GetInputMeans() const
+{
+	return InputMeans;
+}
+
+const TArrayView<const float> UNeuralMorphNetwork::GetInputStds() const
+{
+	return InputStd;
+}
+
+const int32 UNeuralMorphNetwork::GetNumLayers() const
+{
+	return Layers.Num();
+}
+
+UNeuralMorphNetworkLayer& UNeuralMorphNetwork::GetLayer(int32 Index) const
+{
+	return *Layers[Index].Get();
+}
+
+
+//--------------------------------------------------------------------------
+// UNeuralMorphNetworkInstance
+//--------------------------------------------------------------------------
+
+TArrayView<float> UNeuralMorphNetworkInstance::GetInputs()
+{ 
+	return TArrayView<float>(Inputs.GetData(), Inputs.Num());
+}
+
+TArrayView<const float> UNeuralMorphNetworkInstance::GetInputs() const
+{ 
+	return TArrayView<const float>(Inputs.GetData(), Inputs.Num());
+}
+
+TArrayView<float> UNeuralMorphNetworkInstance::GetOutputs()
+{ 
+	return TArrayView<float>(Outputs.GetData(), Outputs.Num());
+}
+
+TArrayView<const float> UNeuralMorphNetworkInstance::GetOutputs() const
+{ 
+	return TArrayView<const float>(Outputs.GetData(), Outputs.Num());
+}
+
+const UNeuralMorphNetwork& UNeuralMorphNetworkInstance::GetNeuralNetwork() const
+{ 
+	return *Network.Get();
+}
+
+void UNeuralMorphNetworkInstance::Init(UNeuralMorphNetwork* InNeuralNetwork)
+{
+	Network = InNeuralNetwork;
+	Inputs.SetNumZeroed(Network->GetNumInputs());
+	Outputs.SetNumZeroed(Network->GetNumOutputs());
+
+	// Find the largest layer unit size and pre-allocate a buffer of that size.
+	int32 MaxNumUnits = 0;
+	const int32 NumLayers = Network->GetNumLayers();
+	for (int32 LayerIndex = 0; LayerIndex < NumLayers; ++LayerIndex)
+	{
+		const UNeuralMorphNetworkLayer& CurLayer = Network->GetLayer(LayerIndex);
+		const int32 NumInputUnits = CurLayer.NumInputs * CurLayer.Depth;
+		const int32 NumOutputUnits = CurLayer.NumOutputs * CurLayer.Depth;
+		MaxNumUnits = FMath::Max3<int32>(NumInputUnits, NumOutputUnits, MaxNumUnits);
+	}
+
+	TempInputArray.SetNumZeroed(MaxNumUnits);
+	TempOutputArray.SetNumZeroed(MaxNumUnits);
+}
+
+void UNeuralMorphNetworkInstance::RunGlobalModel(const FRunSettings& RunSettings)
+{
+	float* RESTRICT TempInputBuffer = RunSettings.TempInputBuffer;
+	float* RESTRICT TempOutputBuffer = RunSettings.TempOutputBuffer;
+
+	const int32 NumLayers = Network->GetNumLayers();
+	for (int32 LayerIndex = 0; LayerIndex < NumLayers; ++LayerIndex)
+	{
+		const UNeuralMorphNetworkLayer& CurLayer = Network->GetLayer(LayerIndex);
+		const int32 NumLayerInputs = CurLayer.NumInputs;
+		const int32 NumLayerOutputs = CurLayer.NumOutputs;
+
+		// Normalize inputs for the first layer inputs.
+		if (LayerIndex == 0)
+		{
+			const float* const RESTRICT NetworkInputs = RunSettings.InputBuffer;
+			const float* const RESTRICT Means = RunSettings.InputMeansBuffer;
+			const float* const RESTRICT Stds = RunSettings.InputStdsBuffer;
+			for (int32 Index = 0; Index < NumLayerInputs; ++Index)
+			{
+				TempInputBuffer[Index] = (NetworkInputs[Index] - Means[Index]) / Stds[Index];
+			}
+		}
+
+		// If we reached the last layer, output to the final buffer.
+		if (LayerIndex == NumLayers - 1)
+		{
+			TempOutputBuffer = RunSettings.OutputBuffer;
+		}
+
+		#if NEURALMORPHMODEL_USE_ISPC
+			const float* const RESTRICT LayerWeights = CurLayer.Weights.GetData();
+			const float* const RESTRICT LayerBiases = CurLayer.Biases.GetData();
+			ispc::MorphNeuralNetwork_LayerForward(TempOutputBuffer, TempInputBuffer, LayerWeights, LayerBiases, NumLayerInputs, NumLayerOutputs);
+		#else
+			// Init the outputs to the bias.
+			const float* const RESTRICT LayerBiases = CurLayer.Biases.GetData();
+			for (int32 Index = 0; Index < NumLayerOutputs; ++Index)
+			{
+				TempOutputBuffer[Index] = LayerBiases[Index];
+			}
+
+			// Multiply layer inputs with the weights.
+			const float* const RESTRICT LayerWeights = CurLayer.Weights.GetData();
+			for (int32 InputIndex = 0; InputIndex < NumLayerInputs; ++InputIndex)
+			{
+				const float InputValue = TempInputBuffer[InputIndex];
+				const int32 Offset = InputIndex * NumLayerOutputs;
+				for (int32 OutputIndex = 0; OutputIndex < NumLayerOutputs; ++OutputIndex)
+				{
+					TempOutputBuffer[OutputIndex] += InputValue * LayerWeights[Offset + OutputIndex];
+				}
+			}
+
+			// Apply ELU activation.
+			for (int32 Index = 0; Index < NumLayerOutputs; ++Index)
+			{
+				const float X = TempOutputBuffer[Index];
+				TempOutputBuffer[Index] = (X > 0.0f) ? X : FMath::InvExpApprox(-X) - 1.0f;
+			}
+		#endif
+
+		// The outputs are now input to the next layer.
+		if (LayerIndex < NumLayers - 1)
+		{
+			// Swap the inputs and outputs, as the outputs are the input to the next layer.
+			float* const SwapTemp = TempInputBuffer;
+			TempInputBuffer = TempOutputBuffer;
+			TempOutputBuffer = SwapTemp;
+		}
+	}
+}
+
+
+void UNeuralMorphNetworkInstance::RunLocalModel(const FRunSettings& RunSettings)
+{
+	float* RESTRICT TempInputBuffer = RunSettings.TempInputBuffer;
+	float* RESTRICT TempOutputBuffer = RunSettings.TempOutputBuffer;
+
+	const int32 NumLayers = Network->GetNumLayers();
+	for (int32 LayerIndex = 0; LayerIndex < NumLayers; ++LayerIndex)
+	{
+		const UNeuralMorphNetworkLayer& CurLayer = Network->GetLayer(LayerIndex);
+		const int32 NumInputsPerBone = CurLayer.NumInputs;
+		const int32 NumBones = CurLayer.Depth;
+
+		// Normalize inputs for the first layer inputs.
+		if (LayerIndex == 0)
+		{
+			const float* const RESTRICT NetworkInputs = RunSettings.InputBuffer;
+			const float* const RESTRICT Means = RunSettings.InputMeansBuffer;
+			const float* const RESTRICT Stds = RunSettings.InputStdsBuffer;
+			const int32 NumInputs = NumBones * NumInputsPerBone;
+			for (int32 Index = 0; Index < NumInputs; ++Index)
+			{
+				TempInputBuffer[Index] = (NetworkInputs[Index] - Means[Index]) / Stds[Index];
+			}
+		}
+
+		// If we reached the last layer, output to the final buffer.
+		if (LayerIndex == NumLayers - 1)
+		{
+			TempOutputBuffer = RunSettings.OutputBuffer;
+		}
+
+		// Init the output buffer to the bias values.
+		const float* const RESTRICT LayerBiases = CurLayer.Biases.GetData();
+		const int32 NumOutputsPerBone = CurLayer.NumOutputs;
+		const int32 NumOutputs = NumBones * NumOutputsPerBone;
+		for (int32 Index = 0; Index < NumOutputs; ++Index)
+		{
+			TempOutputBuffer[Index] = LayerBiases[Index];
+		}
+
+		const float* const RESTRICT LayerWeights = CurLayer.Weights.GetData();
+		for (int32 BoneIndex = 0; BoneIndex < NumBones; ++BoneIndex)
+		{
+			const int32 BoneOutputOffset = BoneIndex * NumOutputsPerBone;
+			const int32 BoneInputOffset = BoneIndex * NumInputsPerBone;
+
+			// Multiply layer inputs with the weights.
+			const int32 WeightOffset = BoneIndex * (NumInputsPerBone * NumOutputsPerBone);
+			for (int32 InputIndex = 0; InputIndex < NumInputsPerBone; ++InputIndex)
+			{
+				const float InputValue = TempInputBuffer[BoneInputOffset + InputIndex];
+				const int32 InputOffset = InputIndex * NumOutputsPerBone;
+				for (int32 OutputIndex = 0; OutputIndex < NumOutputsPerBone; ++OutputIndex)
+				{
+					TempOutputBuffer[BoneOutputOffset + OutputIndex] += InputValue * LayerWeights[WeightOffset + InputOffset + OutputIndex];
+				}
+			}
+		} // For all bones.
+
+		// Apply ELU activation.
+		#if NEURALMORPHMODEL_USE_ISPC
+			ispc::MorphNeuralNetwork_Activation_ELU(TempOutputBuffer, NumOutputs);
+		#else
+			for (int32 OutputIndex = 0; OutputIndex < NumOutputs; ++OutputIndex)
+			{
+				const float X = TempOutputBuffer[OutputIndex];
+				TempOutputBuffer[OutputIndex] = (X > 0.0f) ? X : FMath::InvExpApprox(-X) - 1.0f;
+			}
+		#endif
+
+		// The outputs are now input to the next layer.
+		if (LayerIndex < NumLayers - 1)
+		{
+			// Swap the inputs and outputs, as the outputs are the input to the next layer.
+			float* const SwapTemp = TempInputBuffer;
+			TempInputBuffer = TempOutputBuffer;
+			TempOutputBuffer = SwapTemp;
+		}
+	} // For all layers.
+}
+
+
+void UNeuralMorphNetworkInstance::Run()
+{	
+	TRACE_CPUPROFILER_EVENT_SCOPE(UNeuralMorphNetwork::Run)
+
+	// Setup the buffer pointers.
+	FRunSettings RunSettings;
+	RunSettings.TempInputBuffer  = TempInputArray.GetData();
+	RunSettings.TempOutputBuffer = TempOutputArray.GetData();
+	RunSettings.InputBuffer		 = Inputs.GetData();
+	RunSettings.InputMeansBuffer = Network->GetInputMeans().GetData();
+	RunSettings.InputStdsBuffer  = Network->GetInputStds().GetData();
+	RunSettings.OutputBuffer	 = Outputs.GetData();
+
+	// Run the network.
+	if (Network->GetMode() == ENeuralMorphMode::Global)
+	{
+		RunGlobalModel(RunSettings);
+	}
+	else
+	{
+		checkSlow(Network->GetMode() == ENeuralMorphMode::Local)
+		RunLocalModel(RunSettings);
+	}
+}
diff --git a/Engine/Plugins/Animation/MLDeformer/NeuralMorphModel/Source/NeuralMorphModel/Private/NeuralMorphNetwork.ispc b/Engine/Plugins/Animation/MLDeformer/NeuralMorphModel/Source/NeuralMorphModel/Private/NeuralMorphNetwork.ispc
new file mode 100644
--- /dev/null
+++ b/Engine/Plugins/Animation/MLDeformer/NeuralMorphModel/Source/NeuralMorphModel/Private/NeuralMorphNetwork.ispc
@@ -0,0 +1,42 @@
+// Copyright Epic Games, Inc. All Rights Reserved.
+
+float InvExpApprox(const float X)
+{
+	return 1 / (1 + 1.00746054f * X + 0.45053901f * X * X + 0.25724632f * X * X * X);		
+}
+
+export void MorphNeuralNetwork_LayerForward(
+	uniform float Outputs[],
+	const uniform float Inputs[],
+	const uniform float Weights[],
+	const uniform float Biases[],
+	const uniform int NumRows,
+	const uniform int NumCols)
+{
+	foreach (Index = 0 ... NumCols)
+	{
+		Outputs[Index] = Biases[Index];
+	}
+
+	for (uniform int Row = 0; Row < NumRows; Row++)
+	{
+		uniform float X = Inputs[Row];
+		foreach (Col = 0 ... NumCols)
+		{
+			Outputs[Col] += X * Weights[Row * NumCols + Col];
+		}
+	}
+
+	foreach (Index = 0 ... NumCols)
+	{
+		Outputs[Index] = Outputs[Index] > 0.0f ? Outputs[Index] : InvExpApprox(-Outputs[Index]) - 1.0f;
+	}
+}
+
+export void MorphNeuralNetwork_Activation_ELU(uniform float InputsOutputs[], const uniform int Num)
+{
+	foreach (Index = 0 ... Num)
+	{
+		InputsOutputs[Index] = InputsOutputs[Index] > 0.0f ? InputsOutputs[Index] : InvExpApprox(-InputsOutputs[Index]) - 1.0f;
+	}
+}
diff --git a/Engine/Plugins/Animation/MLDeformer/NeuralMorphModel/Source/NeuralMorphModel/Public/NeuralMorphModel.h b/Engine/Plugins/Animation/MLDeformer/NeuralMorphModel/Source/NeuralMorphModel/Public/NeuralMorphModel.h
--- a/Engine/Plugins/Animation/MLDeformer/NeuralMorphModel/Source/NeuralMorphModel/Public/NeuralMorphModel.h
+++ b/Engine/Plugins/Animation/MLDeformer/NeuralMorphModel/Source/NeuralMorphModel/Public/NeuralMorphModel.h
@@ -9,12 +9,26 @@
 class USkeletalMesh;
 class UGeometryCache;
 class UAnimSequence;
+class UNeuralMorphNetwork;
 class UNeuralNetwork;
 class UMLDeformerAsset;
 class UMLDeformerModelInstance;
 class USkeleton;
 struct FExternalMorphSet;
 
+/** Set this to 1 if you like to force to use NNI for inference, rather than the (faster) custom inference, otherwise set to 0. */
+#define NEURALMORPHMODEL_FORCE_USE_NNI 0
+
+/** 
+ * Specify whether we want to use ISPC if available.
+ * You can comment out this line, to force disable the use of ISPC in this plugin.
+ */
+#define NEURALMORPHMODEL_USE_ISPC INTEL_ISPC
+#if !defined(NEURALMORPHMODEL_USE_ISPC)
+	#define NEURALMORPHMODEL_USE_ISPC 0
+#endif
+
+// Declare our log category.
 NEURALMORPHMODEL_API DECLARE_LOG_CATEGORY_EXTERN(LogNeuralMorphModel, Log, All);
 
 /** 
@@ -61,20 +75,28 @@ class NEURALMORPHMODEL_API UNeuralMorphModel
 public:
 	UNeuralMorphModel(const FObjectInitializer& ObjectInitializer);
 
+	// UObject overrides.
+	virtual void Serialize(FArchive& Archive) override;
+	virtual void PostLoad() override;
+	// ~END UObject overrides.
+
 	// UMLDeformerModel overrides.
-	virtual FString GetDisplayName() const override		{ return "Neural Morph Model"; }
+	virtual FString GetDisplayName() const override			{ return "Neural Morph Model"; }
+	virtual UMLDeformerModelInstance* CreateModelInstance(UMLDeformerComponent* Component) override;
 	// ~END UMLDeformerModel overrides.
 
-#if WITH_EDITORONLY_DATA
-	int32 GetLocalNumHiddenLayers() const				{ return LocalNumHiddenLayers; }
-	int32 GetLocalNumNeuronsPerLayer() const			{ return LocalNumNeuronsPerLayer; }
-	int32 GetGlobalNumHiddenLayers() const				{ return GlobalNumHiddenLayers; }
-	int32 GetGlobalNumNeuronsPerLayer() const			{ return GlobalNumNeuronsPerLayer; }
-	int32 GetNumIterations() const						{ return NumIterations; }
-	int32 GetBatchSize() const							{ return BatchSize; }
-	float GetLearningRate() const						{ return LearningRate; }
-	float GetLearningRateDecay() const					{ return LearningRateDecay; }
-	float GetRegularizationFactor() const				{ return RegularizationFactor;  }
+	UNeuralMorphNetwork* GetNeuralMorphNetwork() const		{ return NeuralMorphNetwork.Get(); }
+	void SetNeuralMorphNetwork(UNeuralMorphNetwork* Net)	{ NeuralMorphNetwork = Net; }
+	ENeuralMorphMode GetModelMode() const					{ return Mode; }
+	int32 GetLocalNumHiddenLayers() const					{ return LocalNumHiddenLayers; }
+	int32 GetLocalNumNeuronsPerLayer() const				{ return LocalNumNeuronsPerLayer; }
+	int32 GetGlobalNumHiddenLayers() const					{ return GlobalNumHiddenLayers; }
+	int32 GetGlobalNumNeuronsPerLayer() const				{ return GlobalNumNeuronsPerLayer; }
+	int32 GetNumIterations() const							{ return NumIterations; }
+	int32 GetBatchSize() const								{ return BatchSize; }
+	float GetLearningRate() const							{ return LearningRate; }
+	float GetLearningRateDecay() const						{ return LearningRateDecay; }
+	float GetRegularizationFactor() const					{ return RegularizationFactor;  }
 
 public:
 	/**
@@ -113,7 +135,7 @@ public:
 	 * Higher numbers will slow down performance but can deal with more complex deformations. 
 	 * For the local model you most likely want to stick with a value of one or two.
 	 */
-	UPROPERTY(EditAnywhere, BlueprintReadWrite, AdvancedDisplay, DisplayName = "Num Hidden Layers", Category = "Training Settings", meta = (ClampMin = "1", ClampMax = "5"))
+	UPROPERTY(EditAnywhere, BlueprintReadWrite, AdvancedDisplay, DisplayName = "Num Hidden Layers", Category = "Training Settings", meta = (ClampMin = "0", ClampMax = "5"))
 	int32 LocalNumHiddenLayers = 1;
 
 	/** 
@@ -121,11 +143,11 @@ public:
 	 * Higher numbers will slow down performance but allow for more complex mesh deformations. 
 	 * For the local mode you probably want to keep this around the same value as the number of morph targets per bone.
 	 */
-	UPROPERTY(EditAnywhere, BlueprintReadWrite, AdvancedDisplay, DisplayName = "Num Neurons Per Layer", Category = "Training Settings", meta = (ClampMin = "1", ClampMax = "100"))
+	UPROPERTY(EditAnywhere, BlueprintReadWrite, AdvancedDisplay, DisplayName = "Num Neurons Per Layer", Category = "Training Settings", meta = (ClampMin = "1", ClampMax = "4096"))
 	int32 LocalNumNeuronsPerLayer = 6;
 
 	/** The number of hidden layers that the neural network model will have.\nHigher numbers will slow down performance but can deal with more complex deformations. */
-	UPROPERTY(EditAnywhere, BlueprintReadWrite, AdvancedDisplay, DisplayName = "Num Hidden Layers", Category = "Training Settings", meta = (ClampMin = "1", ClampMax = "5"))
+	UPROPERTY(EditAnywhere, BlueprintReadWrite, AdvancedDisplay, DisplayName = "Num Hidden Layers", Category = "Training Settings", meta = (ClampMin = "0", ClampMax = "6"))
 	int32 GlobalNumHiddenLayers = 2;
 
 	/** The number of units/neurons per hidden layer. Higher numbers will slow down performance but allow for more complex mesh deformations. */
@@ -155,5 +177,11 @@ public:
 	 */
 	UPROPERTY(EditAnywhere, BlueprintReadWrite, AdvancedDisplay, Category = "Training Settings", meta = (ClampMin = "0.0", ClampMax = "10.0"))
 	float RegularizationFactor = 1.0f;
-#endif // WITH_EDITORONLY_DATA
+
+	/**
+	 * The neural morph model network.
+	 * This is a specialized network used for inference, for performance reasons.
+	 */
+	UPROPERTY()
+	TObjectPtr<UNeuralMorphNetwork> NeuralMorphNetwork;
 };
diff --git a/Engine/Plugins/Animation/MLDeformer/NeuralMorphModel/Source/NeuralMorphModel/Public/NeuralMorphModelInstance.h b/Engine/Plugins/Animation/MLDeformer/NeuralMorphModel/Source/NeuralMorphModel/Public/NeuralMorphModelInstance.h
new file mode 100644
--- /dev/null
+++ b/Engine/Plugins/Animation/MLDeformer/NeuralMorphModel/Source/NeuralMorphModel/Public/NeuralMorphModelInstance.h
@@ -0,0 +1,26 @@
+// Copyright Epic Games, Inc. All Rights Reserved.
+
+#pragma once
+
+#include "MLDeformerMorphModelInstance.h"
+#include "NeuralMorphModelInstance.generated.h"
+
+class UNeuralMorphNetworkInstance;
+
+UCLASS()
+class NEURALMORPHMODEL_API UNeuralMorphModelInstance
+	: public UMLDeformerMorphModelInstance
+{
+	GENERATED_BODY()
+
+public:
+	// UMLDeformerModelInstance overrides.
+	virtual void Init(USkeletalMeshComponent* SkelMeshComponent) override;
+	virtual bool SetupInputs() override;
+	virtual void Execute(float ModelWeight) override;
+	// ~END UMLDeformerModelInstance overrides.
+
+protected:
+	UPROPERTY(Transient)
+	TObjectPtr<UNeuralMorphNetworkInstance> NetworkInstance;
+};
diff --git a/Engine/Plugins/Animation/MLDeformer/NeuralMorphModel/Source/NeuralMorphModel/Public/NeuralMorphNetwork.h b/Engine/Plugins/Animation/MLDeformer/NeuralMorphModel/Source/NeuralMorphModel/Public/NeuralMorphNetwork.h
new file mode 100644
--- /dev/null
+++ b/Engine/Plugins/Animation/MLDeformer/NeuralMorphModel/Source/NeuralMorphModel/Public/NeuralMorphNetwork.h
@@ -0,0 +1,256 @@
+// Copyright Epic Games, Inc. All Rights Reserved.
+#pragma once
+
+#include "CoreTypes.h"
+#include "UObject/Object.h"
+#include "NeuralMorphModel.h"
+#include "NeuralMorphNetwork.generated.h"
+
+class UNeuralMorphNetworkInstance;
+
+/** A fully connected layer, which contains the weights and biases for those connections. */ 
+UCLASS()
+class NEURALMORPHMODEL_API UNeuralMorphNetworkLayer
+	: public UObject
+{
+	GENERATED_BODY()
+
+public:
+	/** The weight matrix number of inputs (rows). */
+	UPROPERTY()
+	int32 NumInputs = 0;
+
+	/** The weight matrix number of outputs (columns). */
+	UPROPERTY()
+	int32 NumOutputs = 0;
+
+	/** The third dimension of the layer. This basically contains the number of bones in local mode. */
+	UPROPERTY()
+	int32 Depth = 1;
+
+	/** The weights, which is basically a 2d array. The number of weights will be equal to the number of rows multiplied by number of columns multiplied by the depth. */
+	UPROPERTY()
+	TArray<float> Weights;
+
+	/** The biases. The number of biases will be the same as the number of columns multiplied by the depth. */
+	UPROPERTY()
+	TArray<float> Biases;
+};
+
+/**
+ * The specialized neural network for the Neural Morph Model.
+ * This class is used to do inference at runtime at a higher performance than using UNeuralNetwork.
+ * The reason why it is faster is because it is a highly specialized network for this specific model.
+ */
+UCLASS()
+class NEURALMORPHMODEL_API UNeuralMorphNetwork
+	: public UObject
+{
+	GENERATED_BODY()
+
+public:
+	/** Clear the network, getting rid of all weights and biases. */
+	void Empty();
+
+	/** 
+	 * Check if the network is empty or not.
+	 * If it is empty, it means it hasn't been loaded, and cannot do anything.
+	 * Empty means it has no inputs, outputs or weights.
+	 * @return Returns true if empty, otherwise false is returned.
+	 */
+	bool IsEmpty() const;
+
+	/**
+	 * Load the network from a file on disk.
+	 * When loading fails, the network will be emptied.
+	 * @param Filename The file path to load from.
+	 * @return Returns true when successfully loaded this model, otherwise false is returned.
+	 */
+	bool Load(const FString& Filename);
+
+	/**
+	 * Create an instance of this neural network.
+	 * @return A pointer to the neural network instance.
+	 */
+	UNeuralMorphNetworkInstance* CreateInstance();
+
+	/**
+	 * Get the number of inputs, which is the number of floats the network takes as input.
+	 * @result The number of input floats to the network.
+	 */
+	int32 GetNumInputs() const;
+
+	/**
+	 * Get the number of outputs, which is the number of floats the network will output.
+	 * @return The number of floats that the network outputs.
+	 */
+	int32 GetNumOutputs() const;
+
+	/**
+	 * Get the number of bones that are input to the network.
+	 * @return The number of bones.
+	 */
+	int32 GetNumBones() const;
+
+	/**
+	 * Get the number of curves that are input to the network.
+	 * @return The number of curves.
+	 */
+	int32 GetNumCurves() const;
+
+	/**
+	 * Get the number of morph targets per bone. This will only be valid if the GetMode() method returns the local model mode.
+	 * For global mode networks, this will return 0.
+	 * @return The number of morph targets per bone.
+	 */
+	int32 GetNumMorphsPerBone() const;
+
+	/**
+	 * Get the mode that the model was trained for, either global or local mode.
+	 * @return The mode the model was trained for.
+	 */
+	ENeuralMorphMode GetMode() const;
+
+	/**
+	 * Get the means of each input, used for normalizing the input values.
+	 * @return The array of mean values, one value for each network input.
+	 */
+	const TArrayView<const float> GetInputMeans() const;
+
+	/**
+	 * Get the standard deviations of each input, used for normalizing the input values.
+	 * @return The array that contains the standard deviations, one value for each network input.
+	 */
+	const TArrayView<const float> GetInputStds() const;
+
+	/**
+	 * Get the number of network layers.
+	 * This equals to the number of hidden layers plus one.
+	 * @return The number of network layer.
+	 */
+	const int32 GetNumLayers() const;
+
+	/**
+	 * Get a given network layer.
+	 * @return A reference to the layer, which will contain the weights and biases.
+	 */
+	UNeuralMorphNetworkLayer& GetLayer(int32 Index) const;
+
+private:
+	/** The network weights and biases, between the different layers. */
+	UPROPERTY()
+	TArray<TObjectPtr<UNeuralMorphNetworkLayer>> Layers;
+
+	/** The means of the input values, used to normalize inputs. */
+	UPROPERTY()
+	TArray<float> InputMeans;
+
+	/** The standard deviation of the input values, used to normalize inputs. */
+	UPROPERTY()
+	TArray<float> InputStd;
+
+	/** The mode of the network, either local or global. */
+	UPROPERTY()
+	ENeuralMorphMode Mode = ENeuralMorphMode::Global;
+
+	/** The number of morph targets per bone, if set Mode == Local, otherwise ignored. */
+	UPROPERTY()
+	int32 NumMorphsPerBone = 0;
+
+	/** The number of bones that were input. */
+	UPROPERTY()
+	int32 NumBones = 0;
+
+	/** The number of curves that were input. */
+	UPROPERTY()
+	int32 NumCurves = 0;
+};
+
+/** 
+ * An instance of a UNeuralMorphNetwork.
+ * The instance holds its own input and output buffers and only reads from the network object it was instanced from.
+ * This allows it to be multithreaded.
+ */
+UCLASS()
+class NEURALMORPHMODEL_API UNeuralMorphNetworkInstance
+	: public UObject
+{
+	friend class UNeuralMorphNetwork;
+
+	GENERATED_BODY()
+
+public:
+	/**
+	 * Get the network input buffer.
+	 * @return The array view of floats that represents the network inputs.
+	 */
+	TArrayView<float> GetInputs();
+	TArrayView<const float> GetInputs() const;
+
+	/**
+	 * Get the network output buffer.
+	 * @return The array view of floats that represents the network outputs.
+	 */
+	TArrayView<float> GetOutputs();
+	TArrayView<const float> GetOutputs() const;
+
+	/**
+	 * Get the neural network this is an instance of.
+	 * @return A reference to the neural network.
+	 */
+	const UNeuralMorphNetwork& GetNeuralNetwork() const;
+
+	/**
+	 * Run the neural network, performing inference.
+	 * This will update the values in the output buffer that you can get with the GetOutputs method.
+	 * This also assumes you have set all the right input values already.
+	 */
+	void Run();
+
+private:
+	struct FRunSettings
+	{
+		float* TempInputBuffer;
+		float* TempOutputBuffer;
+		const float* InputStdsBuffer;
+		const float* InputMeansBuffer;
+		const float* InputBuffer;
+		float* OutputBuffer;
+	};
+
+	/**
+ 	 * Run inference using the global model mode, which is one fully connected MLP.
+	 * @param RunSettings The settings used to run the model. This specifies a set of buffers.
+	 */
+	void RunGlobalModel(const FRunSettings& RunSettings);
+
+	/**
+ 	 * Run inference using the local model mode, which basically has a small MLP for every bone.
+	 * @param RunSettings The settings used to run the model. This specifies a set of buffers.
+	 */
+	void RunLocalModel(const FRunSettings& RunSettings);
+
+	/**
+ 	 * Initialize this instance using a given network object.
+	 * This is automatically called by the UNeuralMorphNetwork::CreateInstance() method.
+	 * @param InNeuralNetwork The network that this is an instance of.
+	 */
+	void Init(UNeuralMorphNetwork* InNeuralNetwork);
+
+private:
+	/** The input values. */
+	TArray<float> Inputs;
+
+	/** The output values. */
+	TArray<float> Outputs;
+
+	/** A pre-allocated temp buffer for inputs. */
+	TArray<float> TempInputArray;
+
+	/** A pre-allocated temp buffer for outputs. */
+	TArray<float> TempOutputArray;
+
+	/** The neural network this is an instance of. */
+	UPROPERTY(Transient)
+	TObjectPtr<UNeuralMorphNetwork> Network;
+};
diff --git a/Engine/Plugins/Animation/MLDeformer/NeuralMorphModel/Source/NeuralMorphModelEditor/Private/NeuralMorphEditorModel.cpp b/Engine/Plugins/Animation/MLDeformer/NeuralMorphModel/Source/NeuralMorphModelEditor/Private/NeuralMorphEditorModel.cpp
--- a/Engine/Plugins/Animation/MLDeformer/NeuralMorphModel/Source/NeuralMorphModelEditor/Private/NeuralMorphEditorModel.cpp
+++ b/Engine/Plugins/Animation/MLDeformer/NeuralMorphModel/Source/NeuralMorphModelEditor/Private/NeuralMorphEditorModel.cpp
@@ -3,6 +3,7 @@
 #include "NeuralMorphEditorModel.h"
 #include "IDetailsView.h"
 #include "NeuralMorphModel.h"
+#include "NeuralMorphNetwork.h"
 #include "MLDeformerMorphModelVizSettings.h"
 #include "NeuralMorphTrainingModel.h"
 #include "MLDeformerEditorToolkit.h"
@@ -40,10 +41,58 @@ namespace UE::NeuralMorphModel
 		}
 	}
 
+	bool FNeuralMorphEditorModel::IsTrained() const
+	{
+#if NEURALMORPHMODEL_FORCE_USE_NNI
+		return GetMorphModel()->GetNeuralNetwork() != nullptr;
+#else
+		return (GetMorphModel()->GetNeuralNetwork() != nullptr) || (GetNeuralMorphModel()->GetNeuralMorphNetwork() != nullptr);
+#endif
+	}
+
 	ETrainingResult FNeuralMorphEditorModel::Train()
 	{
 		return TrainModel<UNeuralMorphTrainingModel>(this);
 	}
+
+	bool FNeuralMorphEditorModel::LoadTrainedNetwork() const
+	{
+#if NEURALMORPHMODEL_FORCE_USE_NNI
+		const bool bSuccess = FMLDeformerMorphModelEditorModel::LoadTrainedNetwork();
+		if (bSuccess)
+		{
+			GetNeuralMorphModel()->SetNeuralMorphNetwork(nullptr);	// Force disable custom inference.
+		}
+		return bSuccess;
+#else
+		// Load the specialized neural morph model network.
+		// Base the filename on the onnx filename, and replace the file extension.
+		FString NetworkFilename = GetTrainedNetworkOnnxFile();
+		NetworkFilename.RemoveFromEnd(TEXT("onnx"));
+		NetworkFilename += TEXT("nmn");
+
+		// Load the actual network.
+		UNeuralMorphNetwork* NeuralNet = NewObject<UNeuralMorphNetwork>(Model);
+		if (!NeuralNet->Load(NetworkFilename))
+		{
+			UE_LOG(LogNeuralMorphModel, Error, TEXT("Failed to load neural morph network from file '%s'!"), *NetworkFilename);
+			NeuralNet = nullptr;
+
+			// Restore the deltas to the ones before training started.
+			GetMorphModel()->SetMorphTargetDeltas(MorphTargetDeltasBackup);
+			return false;
+		}
+
+		if (NeuralNet->IsEmpty())
+		{
+			NeuralNet = nullptr;
+		}
+
+		GetNeuralMorphModel()->SetNeuralNetwork(nullptr);			// Disable NNI inference.
+		GetNeuralMorphModel()->SetNeuralMorphNetwork(NeuralNet);	// Use our custom inference.
+#endif
+		return true;
+	}
 }	// namespace UE::NeuralMorphModel
 
 #undef LOCTEXT_NAMESPACE
diff --git a/Engine/Plugins/Animation/MLDeformer/NeuralMorphModel/Source/NeuralMorphModelEditor/Private/NeuralMorphModelDetails.cpp b/Engine/Plugins/Animation/MLDeformer/NeuralMorphModel/Source/NeuralMorphModelEditor/Private/NeuralMorphModelDetails.cpp
--- a/Engine/Plugins/Animation/MLDeformer/NeuralMorphModel/Source/NeuralMorphModelEditor/Private/NeuralMorphModelDetails.cpp
+++ b/Engine/Plugins/Animation/MLDeformer/NeuralMorphModel/Source/NeuralMorphModelEditor/Private/NeuralMorphModelDetails.cpp
@@ -5,6 +5,11 @@
 #include "DetailCategoryBuilder.h"
 #include "DetailWidgetRow.h"
 
+#include "Widgets/Layout/SBox.h"
+#include "Widgets/Input/SButton.h"
+#include "SWarningOrErrorBox.h"
+
+
 #define LOCTEXT_NAMESPACE "NeuralMorphModelDetails"
 
 namespace UE::NeuralMorphModel
@@ -48,6 +53,42 @@ namespace UE::NeuralMorphModel
 		TrainingSettingsCategoryBuilder->AddProperty(GET_MEMBER_NAME_CHECKED(UNeuralMorphModel, LearningRateDecay), UNeuralMorphModel::StaticClass());
 		TrainingSettingsCategoryBuilder->AddProperty(GET_MEMBER_NAME_CHECKED(UNeuralMorphModel, RegularizationFactor), UNeuralMorphModel::StaticClass());
 	}
+
+	void FNeuralMorphModelDetails::AddTrainingSettingsErrors()
+	{
+		UNeuralMorphModel* NeuralMorphModel = Cast<UNeuralMorphModel>(Model);
+		check(NeuralMorphModel);
+
+		#if !NEURALMORPHMODEL_FORCE_USE_NNI
+			const bool bShowWarning = (NeuralMorphModel->GetNeuralNetwork() != nullptr) && (NeuralMorphModel->GetNeuralMorphNetwork() == nullptr);
+			FDetailWidgetRow& PerformanceWarningRow = TrainingSettingsCategoryBuilder->AddCustomRow(FText::FromString("NeuralNetPerformanceWarning"))
+				.Visibility(bShowWarning ? EVisibility::Visible : EVisibility::Collapsed)
+				.WholeRowContent()
+				[
+					SNew(SBox)
+					.Padding(FMargin(0.0f, 4.0f))
+					[
+						SNew(SWarningOrErrorBox)
+						.MessageStyle(EMessageStyle::Warning)
+						.Message(LOCTEXT("NeuralNetPerformanceWarning", "The model must be retrained in order to make use of higher performance inference."))
+					]
+				];
+		#else
+			const bool bShowWarning = (NeuralMorphModel->GetNeuralNetwork() == nullptr) && (NeuralMorphModel->GetNeuralMorphNetwork() != nullptr);
+			FDetailWidgetRow& PerformanceWarningRow = TrainingSettingsCategoryBuilder->AddCustomRow(FText::FromString("NeuralNetWrongInferenceWarning"))
+				.Visibility(bShowWarning ? EVisibility::Visible : EVisibility::Collapsed)
+				.WholeRowContent()
+				[
+					SNew(SBox)
+					.Padding(FMargin(0.0f, 4.0f))
+					[
+						SNew(SWarningOrErrorBox)
+						.MessageStyle(EMessageStyle::Warning)
+						.Message(LOCTEXT("NeuralNetWrongInferenceWarning", "The model was trained using custom inference, but NNI is set as active inference engine. Please retrain the model or switch back to custom inference by recompiling the plugin."))
+					]
+				];
+		#endif
+	}
 }	// namespace UE::NeuralMorphModel
 
 #undef LOCTEXT_NAMESPACE
diff --git a/Engine/Plugins/Animation/MLDeformer/NeuralMorphModel/Source/NeuralMorphModelEditor/Public/NeuralMorphEditorModel.h b/Engine/Plugins/Animation/MLDeformer/NeuralMorphModel/Source/NeuralMorphModelEditor/Public/NeuralMorphEditorModel.h
--- a/Engine/Plugins/Animation/MLDeformer/NeuralMorphModel/Source/NeuralMorphModelEditor/Public/NeuralMorphEditorModel.h
+++ b/Engine/Plugins/Animation/MLDeformer/NeuralMorphModel/Source/NeuralMorphModelEditor/Public/NeuralMorphEditorModel.h
@@ -23,10 +23,12 @@ namespace UE::NeuralMorphModel
 
 		// FGCObject overrides.
 		virtual FString GetReferencerName() const override		{ return TEXT("FNeuralMorphEditorModel"); }
+		virtual bool LoadTrainedNetwork() const override;
 		// ~END FGCObject overrides.
 	
 		// FMLDeformerEditorModel overrides.
 		virtual ETrainingResult Train() override;
+		virtual bool IsTrained() const override;
 		virtual void OnPropertyChanged(FPropertyChangedEvent& PropertyChangedEvent) override;
 		// ~END FMLDeformerEditorModel overrides.
 
diff --git a/Engine/Plugins/Animation/MLDeformer/NeuralMorphModel/Source/NeuralMorphModelEditor/Public/NeuralMorphModelDetails.h b/Engine/Plugins/Animation/MLDeformer/NeuralMorphModel/Source/NeuralMorphModelEditor/Public/NeuralMorphModelDetails.h
--- a/Engine/Plugins/Animation/MLDeformer/NeuralMorphModel/Source/NeuralMorphModelEditor/Public/NeuralMorphModelDetails.h
+++ b/Engine/Plugins/Animation/MLDeformer/NeuralMorphModel/Source/NeuralMorphModelEditor/Public/NeuralMorphModelDetails.h
@@ -23,5 +23,9 @@ namespace UE::NeuralMorphModel
 		// ILayoutDetails overrides.
 		virtual void CustomizeDetails(IDetailLayoutBuilder& DetailBuilder) override;
 		// ~END ILayoutDetails overrides.
+
+		// FMLDeformerModelDetails overrides.
+		virtual void AddTrainingSettingsErrors() override;
+		// ~END FMLDeformerModelDetails overrides.
 	};
 }	// namespace UE::NeuralMorphModel
diff --git a/Engine/Plugins/Animation/MLDeformer/NeuralMorphModel/Source/NeuralMorphModelEditor/Public/NeuralMorphTrainingModel.h b/Engine/Plugins/Animation/MLDeformer/NeuralMorphModel/Source/NeuralMorphModelEditor/Public/NeuralMorphTrainingModel.h
--- a/Engine/Plugins/Animation/MLDeformer/NeuralMorphModel/Source/NeuralMorphModelEditor/Public/NeuralMorphTrainingModel.h
+++ b/Engine/Plugins/Animation/MLDeformer/NeuralMorphModel/Source/NeuralMorphModelEditor/Public/NeuralMorphTrainingModel.h
@@ -4,6 +4,7 @@
 
 #include "CoreMinimal.h"
 #include "MLDeformerTrainingModel.h"
+#include "NeuralMorphModel.h"
 #include "NeuralMorphTrainingModel.generated.h"
 
 /**
@@ -24,4 +25,18 @@ public:
 	 */
 	UFUNCTION(BlueprintImplementableEvent, Category = "Training Model")
 	int32 Train() const;
+
+	/**
+	 * Do we want to force using NNI for inference, instead of our custom inference?
+	 * @result Returns true when NNI should be used for inference, or false if custom (faster) inference should be used.
+	 */
+	UFUNCTION(BlueprintPure, Category = "Training Data")
+	bool ForceUseNNI() const
+	{
+		#if NEURALMORPHMODEL_FORCE_USE_NNI
+			return true;
+		#else
+			return false;
+		#endif
+	}
 };
diff --git a/Engine/Plugins/Animation/MLDeformer/VertexDeltaModel/Source/VertexDeltaModel/Private/Tests/VertexDeltaModelTest.cpp b/Engine/Plugins/Animation/MLDeformer/VertexDeltaModel/Source/VertexDeltaModel/Private/Tests/VertexDeltaModelTest.cpp
--- a/Engine/Plugins/Animation/MLDeformer/VertexDeltaModel/Source/VertexDeltaModel/Private/Tests/VertexDeltaModelTest.cpp
+++ b/Engine/Plugins/Animation/MLDeformer/VertexDeltaModel/Source/VertexDeltaModel/Private/Tests/VertexDeltaModelTest.cpp
@@ -45,8 +45,6 @@ namespace UE::VertexDeltaModelTests
 
 		const int32 NumInputBones = InputInfo->GetNumBones();
 		UTEST_EQUAL(TEXT("InputInfo bone count check"), NumInputBones, 1);
-		UTEST_EQUAL(TEXT("InputInfo bone name string check"), InputInfo->GetBoneNameStrings().Num(), 1);
-		UTEST_EQUAL(TEXT("InputInfo bone name string content check"), InputInfo->GetBoneNameString(0), FString("lowerarm_l"));
 		UTEST_EQUAL(TEXT("InputInfo bone names check"), InputInfo->GetBoneNames().Num(), 1);
 		UTEST_EQUAL(TEXT("InputInfo bone names content check"), InputInfo->GetBoneName(0), FName("lowerarm_l"));
 		UTEST_EQUAL(TEXT("InputInfo vertex count check"), InputInfo->GetNumBaseMeshVertices(), InputInfo->GetNumTargetMeshVertices());
diff --git a/Engine/Source/Runtime/Engine/Classes/Components/SkeletalMeshComponent.h b/Engine/Source/Runtime/Engine/Classes/Components/SkeletalMeshComponent.h
--- a/Engine/Source/Runtime/Engine/Classes/Components/SkeletalMeshComponent.h
+++ b/Engine/Source/Runtime/Engine/Classes/Components/SkeletalMeshComponent.h
@@ -387,6 +387,9 @@ public:
 	// we want to have reference and emptied during evaluation
 	TArray<FTransform> GetBoneSpaceTransforms();
 
+	/** Get the bone space transforms as array view. */
+	TArrayView<const FTransform> GetBoneSpaceTransformsView();
+
 	/** 
 	 * Temporary array of local-space (relative to parent bone) rotation/translation for each bone. 
 	 * This property is not safe to access during evaluation, so we created wrapper.
diff --git a/Engine/Source/Runtime/Engine/Classes/Components/SkinnedMeshComponent.h b/Engine/Source/Runtime/Engine/Classes/Components/SkinnedMeshComponent.h
--- a/Engine/Source/Runtime/Engine/Classes/Components/SkinnedMeshComponent.h
+++ b/Engine/Source/Runtime/Engine/Classes/Components/SkinnedMeshComponent.h
@@ -409,6 +409,9 @@ public:
 	/** Do we have a given set of external morph targets? */
 	bool HasExternalMorphSet(int32 LOD, int32 ID) const;
 
+	/** Check whether a given LOD index is valid for the external morph sets. */
+	bool IsValidExternalMorphSetLODIndex(int32 LOD) const;
+
 	/** Get the external morph sets for a given LOD. */
 	const FExternalMorphSets& GetExternalMorphSets(int32 LOD) const { return ExternalMorphSets[LOD]; }
 
diff --git a/Engine/Source/Runtime/Engine/Private/Components/SkeletalMeshComponent.cpp b/Engine/Source/Runtime/Engine/Private/Components/SkeletalMeshComponent.cpp
--- a/Engine/Source/Runtime/Engine/Private/Components/SkeletalMeshComponent.cpp
+++ b/Engine/Source/Runtime/Engine/Private/Components/SkeletalMeshComponent.cpp
@@ -4663,6 +4663,22 @@ TArray<FTransform> USkeletalMeshComponent::GetBoneSpaceTransforms()
 	PRAGMA_ENABLE_DEPRECATION_WARNINGS
 }
 
+
+TArrayView<const FTransform> USkeletalMeshComponent::GetBoneSpaceTransformsView()
+{
+	// We may be doing parallel evaluation on the current anim instance
+	// Calling this here with true will block this init till that thread completes
+	// and it is safe to continue
+	const bool bBlockOnTask = true; // wait on evaluation task so it is safe to swap the buffers
+	const bool bPerformPostAnimEvaluation = true; // Do PostEvaluation so we make sure to swap the buffers back. 
+	HandleExistingParallelEvaluationTask(bBlockOnTask, bPerformPostAnimEvaluation);
+
+	PRAGMA_DISABLE_DEPRECATION_WARNINGS
+	return BoneSpaceTransforms;
+	PRAGMA_ENABLE_DEPRECATION_WARNINGS
+}
+
+
 #if WITH_EDITOR
 void USkeletalMeshComponent::UpdatePoseWatches()
 {
diff --git a/Engine/Source/Runtime/Engine/Private/Components/SkinnedMeshComponent.cpp b/Engine/Source/Runtime/Engine/Private/Components/SkinnedMeshComponent.cpp
--- a/Engine/Source/Runtime/Engine/Private/Components/SkinnedMeshComponent.cpp
+++ b/Engine/Source/Runtime/Engine/Private/Components/SkinnedMeshComponent.cpp
@@ -680,6 +680,11 @@ FExternalMorphWeightData& USkinnedMeshComponent::GetExternalMorphWeights(int32 L
 	return ExternalMorphWeightData[LOD];
 }
 
+bool USkinnedMeshComponent::IsValidExternalMorphSetLODIndex(int32 LOD) const
+{
+	return ExternalMorphSets.IsValidIndex(LOD);
+}
+
 void USkinnedMeshComponent::AddExternalMorphSet(int32 LOD, int32 ID, TSharedPtr<FExternalMorphSet> MorphSet)
 {
 	const int32 NumLODs = GetNumLODs();
