[MLDeformer]
- Add support for curves to the local model.
- A little optimization in the ISPC
- Disabled some tests because of some NNI problem, and fixed another test.

diff --git a/Engine/Plugins/Animation/MLDeformer/MLDeformerFramework/Source/MLDeformerFramework/Private/MLDeformerInputInfo.cpp b/Engine/Plugins/Animation/MLDeformer/MLDeformerFramework/Source/MLDeformerFramework/Private/MLDeformerInputInfo.cpp
--- a/Engine/Plugins/Animation/MLDeformer/MLDeformerFramework/Source/MLDeformerFramework/Private/MLDeformerInputInfo.cpp
+++ b/Engine/Plugins/Animation/MLDeformer/MLDeformerFramework/Source/MLDeformerFramework/Private/MLDeformerInputInfo.cpp
@@ -205,15 +205,21 @@ FString UMLDeformerInputInfo::GenerateCompatibilityErrorString(USkeletalMesh* In
 }
 
 void UMLDeformerInputInfo::ExtractCurveValues(USkeletalMeshComponent* SkelMeshComponent, TArray<float>& OutValues) const
+{
+	ExtractCurveValues(SkelMeshComponent, OutValues, 1);
+}
+
+void UMLDeformerInputInfo::ExtractCurveValues(USkeletalMeshComponent* SkelMeshComponent, TArray<float>& OutValues, int32 NumFloatsPerCurve) const
 {
 	UAnimInstance* AnimInstance = SkelMeshComponent->GetAnimInstance();
 	const int32 NumCurves = CurveNames.Num();
-	OutValues.Reset(NumCurves);
-	OutValues.AddUninitialized(NumCurves);
+	const int32 NumCurveFloats = NumCurves * NumFloatsPerCurve;
+	OutValues.Reset(NumCurveFloats);
+	OutValues.AddZeroed(NumCurveFloats);
 	for (int32 Index = 0; Index < NumCurves; ++Index)
 	{
 		const FName CurveName = CurveNames[Index];
-		OutValues[Index] = AnimInstance->GetCurveValue(CurveName);
+		OutValues[Index * NumFloatsPerCurve] = AnimInstance->GetCurveValue(CurveName);
 	}
 }
 
@@ -288,10 +294,15 @@ void UMLDeformerInputInfo::ExtractBoneRotations(USkeletalMeshComponent* SkelMesh
 }
 
 int32 UMLDeformerInputInfo::CalcNumNeuralNetInputs() const
+{
+	return CalcNumNeuralNetInputs(6, 1);
+}
+
+int32 UMLDeformerInputInfo::CalcNumNeuralNetInputs(int32 NumFloatsPerBone, int32 NumFloatsPerCurve) const
 {
 	return 
-		BoneNames.Num() * 6 +	// Six floats per bone (2 FVector3 columns of the rotation matrix).
-		CurveNames.Num();		// One float per curve.
+		BoneNames.Num() * NumFloatsPerBone +
+		CurveNames.Num() * NumFloatsPerCurve;
 }
 
 const FSoftObjectPath& UMLDeformerInputInfo::GetSkeletalMesh() const
diff --git a/Engine/Plugins/Animation/MLDeformer/MLDeformerFramework/Source/MLDeformerFramework/Private/MLDeformerModelInstance.cpp b/Engine/Plugins/Animation/MLDeformer/MLDeformerFramework/Source/MLDeformerFramework/Private/MLDeformerModelInstance.cpp
--- a/Engine/Plugins/Animation/MLDeformer/MLDeformerFramework/Source/MLDeformerFramework/Private/MLDeformerModelInstance.cpp
+++ b/Engine/Plugins/Animation/MLDeformer/MLDeformerFramework/Source/MLDeformerFramework/Private/MLDeformerModelInstance.cpp
@@ -151,7 +151,9 @@ FString UMLDeformerModelInstance::CheckCompatibility(USkeletalMeshComponent* InS
 	if (NeuralNetwork && NeuralNetwork->IsLoaded() && Model->GetDeformerAsset())
 	{
 		const int64 NumNeuralNetInputs = NeuralNetwork->GetInputTensor().Num();
-		const int64 NumDeformerAssetInputs = static_cast<int64>(Model->GetInputInfo()->CalcNumNeuralNetInputs());
+		const int32 NumFloatsPerBone = 6;
+		const int32 NumFloatsPerCurve = 1;
+		const int64 NumDeformerAssetInputs = static_cast<int64>(Model->GetInputInfo()->CalcNumNeuralNetInputs(NumFloatsPerBone, NumFloatsPerCurve));
 		if (NumNeuralNetInputs != NumDeformerAssetInputs)
 		{
 			const FString InputErrorString = "The number of network inputs doesn't match the asset. Please retrain the asset."; 
@@ -228,9 +230,8 @@ int64 UMLDeformerModelInstance::SetCurveValues(float* OutputBuffer, int64 Output
 {
 	const UMLDeformerInputInfo* InputInfo = Model->GetInputInfo();
 
-	// Write the weights into the output buffer.
 	int64 Index = StartIndex;
-	const int32 AssetNumCurves = InputInfo->GetNumCurves();	
+	const int32 AssetNumCurves = InputInfo->GetNumCurves();
 	checkf((Index + AssetNumCurves) <= OutputBufferSize, TEXT("Writing curves past the end of the input buffer"));
 
 	// Write the curve weights to the output buffer.
@@ -363,7 +364,9 @@ bool UMLDeformerModelInstance::SetupInputs()
 
 	// If the neural network expects a different number of inputs, do nothing.
 	const int64 NumNeuralNetInputs = NeuralNetwork->GetInputTensorForContext(NeuralNetworkInferenceHandle).Num();
-	const int64 NumDeformerAssetInputs = Model->GetInputInfo()->CalcNumNeuralNetInputs();
+	const int32 NumFloatsPerBone = 6;
+	const int32 NumFloatsPerCurve = 1;
+	const int64 NumDeformerAssetInputs = Model->GetInputInfo()->CalcNumNeuralNetInputs(NumFloatsPerBone, NumFloatsPerCurve);
 	if (NumNeuralNetInputs != NumDeformerAssetInputs)
 	{
 		return false;
diff --git a/Engine/Plugins/Animation/MLDeformer/MLDeformerFramework/Source/MLDeformerFramework/Private/Tests/ModelTests.cpp b/Engine/Plugins/Animation/MLDeformer/MLDeformerFramework/Source/MLDeformerFramework/Private/Tests/ModelTests.cpp
--- a/Engine/Plugins/Animation/MLDeformer/MLDeformerFramework/Source/MLDeformerFramework/Private/Tests/ModelTests.cpp
+++ b/Engine/Plugins/Animation/MLDeformer/MLDeformerFramework/Source/MLDeformerFramework/Private/Tests/ModelTests.cpp
@@ -49,6 +49,9 @@ namespace UE::MLDeformerTests
 		TestModel->UpdateCachedNumVertices();
 		UTEST_EQUAL(TEXT("Imported verts check"), TestModel->GetNumBaseMeshVerts(), ImportedVerts);
 
+		TestModel->GetInputInfo()->SetSkeletalMesh(SkeletalMesh);
+		UTEST_TRUE(TEXT("InputInfo GetSkeletalMesh check"), TestModel->GetInputInfo()->GetSkeletalMesh() == SkeletalMesh);
+
 		const USkeleton* Skeleton = SkeletalMesh->GetSkeleton();
 		bool bErrorOnInvalid = false;
 		const USkeleton* ModelSkeleton = TestModel->GetSkeleton(bErrorOnInvalid, nullptr);
@@ -73,7 +76,7 @@ namespace UE::MLDeformerTests
 		// Add an ML Deformer component and activate it.
 		UMLDeformerComponent* MLDeformerComponent = NewObject<UMLDeformerComponent>(Actor);
 		MLDeformerComponent->SetDeformerAsset(DeformerAsset);
-		MLDeformerComponent->RegisterComponent();
+		MLDeformerComponent->RegisterComponent();		
 		MLDeformerComponent->Activate();
 		UTEST_EQUAL(TEXT("MLDeformerComponent weight check"), MLDeformerComponent->GetWeight(), 1.0f);
 		UTEST_TRUE(TEXT("MLDeformerComponent SkelMeshComponent check"), MLDeformerComponent->GetSkeletalMeshComponent() == SkelMeshComponent);
diff --git a/Engine/Plugins/Animation/MLDeformer/MLDeformerFramework/Source/MLDeformerFramework/Public/MLDeformerInputInfo.h b/Engine/Plugins/Animation/MLDeformer/MLDeformerFramework/Source/MLDeformerFramework/Public/MLDeformerInputInfo.h
--- a/Engine/Plugins/Animation/MLDeformer/MLDeformerFramework/Source/MLDeformerFramework/Public/MLDeformerInputInfo.h
+++ b/Engine/Plugins/Animation/MLDeformer/MLDeformerFramework/Source/MLDeformerFramework/Public/MLDeformerInputInfo.h
@@ -60,11 +60,20 @@ public:
 
 	/**
 	 * Calculate how many inputs this input info generates for the neural network.
-	 * A single bone would take 4 inputs, while a curve takes one input.
+	 * A single bone on default takes 6 inputs.
 	 * @return THe number of input float values to the neural network.
 	 */
+	UE_DEPRECATED(5.2, "Please use the CalcNumNeuralNetInputs that takes parameters instead.")
 	virtual int32 CalcNumNeuralNetInputs() const;
 
+	/**
+	 * Calculate how many inputs this input info generates for the neural network.
+	 * @param NumFloatsPerBone The number of floats used per bone.
+	 * @param NumFloatsPerCurve The number of floats per curve. This can be used when padding is used on curve values.
+	 * @return THe number of input float values to the neural network.
+	 */
+	virtual int32 CalcNumNeuralNetInputs(int32 NumFloatsPerBone, int32 NumFloatsPerCurve) const;
+
 	/**
 	 * Get the number of bones that we trained on.
 	 * @return The number of bones.
@@ -118,12 +127,22 @@ public:
 #endif
 
 	/** 
-	 * Extract the curve values for all curves we're interested in.
+	 * Extract the curve values for all curves we're interested in. Assume one float per curve.
 	 * @param SkelMeshComponent The skeletal mesh component to sample from.
 	 * @param OutValues The array to write the values to. This array will be reset/resized by this method.
 	 */
+	UE_DEPRECATED(5.2, "Please use the ExtractCurveValues method that takes the NumFloatsPerCurve parameter.")
 	virtual void ExtractCurveValues(USkeletalMeshComponent* SkelMeshComponent, TArray<float>& OutValues) const;
 
+	/** 
+	 * Extract the curve values for all curves we're interested in.
+	 * @param SkelMeshComponent The skeletal mesh component to sample from.
+	 * @param OutValues The array to write the values to. This array will be reset/resized by this method.
+	 * @param NumFloatsPerCurve The number of floats per curve. If larger than 1, the remaining floats (after the first one), will be set to 0.
+	 *        So if this is set to 4, you get an array like "0.75, 0, 0, 0", where 0.75 would represent the actual curve value.
+	 */
+	virtual void ExtractCurveValues(USkeletalMeshComponent* SkelMeshComponent, TArray<float>& OutValues, int32 NumFloatsPerCurve) const;
+
 	/**
 	 * Batch convert a set of transform rotation quaternions into two column vectors.
 	 * 
diff --git a/Engine/Plugins/Animation/MLDeformer/MLDeformerFramework/Source/MLDeformerFrameworkEditor/Private/MLDeformerSampler.cpp b/Engine/Plugins/Animation/MLDeformer/MLDeformerFramework/Source/MLDeformerFrameworkEditor/Private/MLDeformerSampler.cpp
--- a/Engine/Plugins/Animation/MLDeformer/MLDeformerFramework/Source/MLDeformerFrameworkEditor/Private/MLDeformerSampler.cpp
+++ b/Engine/Plugins/Animation/MLDeformer/MLDeformerFramework/Source/MLDeformerFrameworkEditor/Private/MLDeformerSampler.cpp
@@ -46,6 +46,7 @@ namespace UE::MLDeformer
 		NumImportedVertices = UMLDeformerModel::ExtractNumImportedSkinnedVertices(Model->GetSkeletalMesh());
 		AnimFrameIndex = 0;
 		SampleTime = 0.0f;
+		NumFloatsPerCurve = 1;
 
 		// Create the actors and components.
 		// This internally skips creating them if they already exist.
@@ -90,7 +91,7 @@ namespace UE::MLDeformer
 	void FMLDeformerSampler::UpdateCurveValues()
 	{
 		const UMLDeformerInputInfo* InputInfo = Model->GetInputInfo();
-		InputInfo->ExtractCurveValues(SkeletalMeshComponent, CurveValues);
+		InputInfo->ExtractCurveValues(SkeletalMeshComponent, CurveValues, NumFloatsPerCurve);
 	}
 
 	void FMLDeformerSampler::Sample(int32 InAnimFrameIndex)
diff --git a/Engine/Plugins/Animation/MLDeformer/MLDeformerFramework/Source/MLDeformerFrameworkEditor/Private/MLDeformerTrainingModel.cpp b/Engine/Plugins/Animation/MLDeformer/MLDeformerFramework/Source/MLDeformerFrameworkEditor/Private/MLDeformerTrainingModel.cpp
--- a/Engine/Plugins/Animation/MLDeformer/MLDeformerFramework/Source/MLDeformerFrameworkEditor/Private/MLDeformerTrainingModel.cpp
+++ b/Engine/Plugins/Animation/MLDeformer/MLDeformerFramework/Source/MLDeformerFrameworkEditor/Private/MLDeformerTrainingModel.cpp
@@ -46,6 +46,11 @@ int32 UMLDeformerTrainingModel::GetNumberSampleDeltas() const
 	return GetModel()->GetNumBaseMeshVerts();
 }
 
+void UMLDeformerTrainingModel::SetNumFloatsPerCurve(int32 NumFloatsPerCurve)
+{
+	EditorModel->GetSampler()->SetNumFloatsPerCurve(NumFloatsPerCurve);
+}
+
 bool UMLDeformerTrainingModel::SetCurrentSampleIndex(int32 Index)
 {
 	return SampleFrame(Index);
diff --git a/Engine/Plugins/Animation/MLDeformer/MLDeformerFramework/Source/MLDeformerFrameworkEditor/Public/MLDeformerSampler.h b/Engine/Plugins/Animation/MLDeformer/MLDeformerFramework/Source/MLDeformerFrameworkEditor/Public/MLDeformerSampler.h
--- a/Engine/Plugins/Animation/MLDeformer/MLDeformerFramework/Source/MLDeformerFrameworkEditor/Public/MLDeformerSampler.h
+++ b/Engine/Plugins/Animation/MLDeformer/MLDeformerFramework/Source/MLDeformerFrameworkEditor/Public/MLDeformerSampler.h
@@ -84,6 +84,15 @@ namespace UE::MLDeformer
 		SIZE_T CalcMemUsagePerFrameInBytes() const;
 		UDebugSkelMeshComponent* GetSkeletalMeshComponent() { return SkeletalMeshComponent; }
 
+		/**
+		 * Set the number of floats used per curve. Values larger than 1 will cause padding with zeros. 
+		 * The first float is always the value. Floats that come after that will contain zeros on default.
+		 */
+		void SetNumFloatsPerCurve(int32 NumFloats)						{ check(NumFloats > 0); NumFloatsPerCurve = NumFloats; }
+
+		/** Get the number of floats per curve. */
+		int32 GetNumFloatsPerCurve() const								{ return NumFloatsPerCurve; }
+
 	protected:
 		/** Create the actors used for sampling. This creates two actors, one for the base skeletal mesh and one for the target mesh. */
 		virtual void CreateActors();
@@ -168,6 +177,13 @@ namespace UE::MLDeformer
 		/** The animation frame we sampled the deltas for. */
 		int32 AnimFrameIndex = -1;
 
+		/**
+		 * The number of floats per curve value. If this is larger than 1, the remaining floats (after the first one) will be set to 0. 
+		 * So basically we will pad the curve values.
+		 * The CurveValues float array will contain something like : "0.75, 0.0, 0.0, 0.0" when this is set to 4.
+		 */
+		int32 NumFloatsPerCurve = 1;
+
 		/** The vertex delta space (pre or post skinning) used when calculating the deltas. */
 		EVertexDeltaSpace VertexDeltaSpace = EVertexDeltaSpace::PreSkinning;
 	};
diff --git a/Engine/Plugins/Animation/MLDeformer/MLDeformerFramework/Source/MLDeformerFrameworkEditor/Public/MLDeformerTrainingModel.h b/Engine/Plugins/Animation/MLDeformer/MLDeformerFramework/Source/MLDeformerFrameworkEditor/Public/MLDeformerTrainingModel.h
--- a/Engine/Plugins/Animation/MLDeformer/MLDeformerFramework/Source/MLDeformerFrameworkEditor/Public/MLDeformerTrainingModel.h
+++ b/Engine/Plugins/Animation/MLDeformer/MLDeformerFramework/Source/MLDeformerFrameworkEditor/Public/MLDeformerTrainingModel.h
@@ -96,6 +96,10 @@ public:
 	UFUNCTION(BlueprintCallable, Category = "Training Data")
 	void SetNeedsResampling(bool bNeedsResampling);
 
+	/** Set the number of floats per curve. On default this is one. */
+	UFUNCTION(BlueprintCallable, Category = "Training Data")
+	void SetNumFloatsPerCurve(int32 NumFloatsPerCurve);
+
 protected:
 	/** This updates the sample deltas, curves, and bone rotations. */
     virtual bool SampleFrame(int32 Index);
diff --git a/Engine/Plugins/Animation/MLDeformer/NeuralMorphModel/Source/NeuralMorphModel/Private/NeuralMorphModelInstance.cpp b/Engine/Plugins/Animation/MLDeformer/NeuralMorphModel/Source/NeuralMorphModel/Private/NeuralMorphModelInstance.cpp
--- a/Engine/Plugins/Animation/MLDeformer/NeuralMorphModel/Source/NeuralMorphModel/Private/NeuralMorphModelInstance.cpp
+++ b/Engine/Plugins/Animation/MLDeformer/NeuralMorphModel/Source/NeuralMorphModel/Private/NeuralMorphModelInstance.cpp
@@ -5,6 +5,7 @@
 #include "NeuralMorphNetwork.h"
 #include "MLDeformerComponent.h"
 #include "MLDeformerAsset.h"
+#include "Animation/AnimInstance.h"
 #include "Components/ExternalMorphSet.h"
 #include "Components/SkeletalMeshComponent.h"
 
@@ -20,6 +21,67 @@ void UNeuralMorphModelInstance::Init(USkeletalMeshComponent* SkelMeshComponent)
 #endif
 }
 
+
+int64 UNeuralMorphModelInstance::SetCurveValues(float* OutputBuffer, int64 OutputBufferSize, int64 StartIndex)
+{
+	UNeuralMorphModel* MorphModel = Cast<UNeuralMorphModel>(Model);
+	UNeuralMorphNetwork* MorphNetwork = MorphModel->GetNeuralMorphNetwork();
+
+	const UMLDeformerInputInfo* InputInfo = Model->GetInputInfo();
+
+	int64 Index = StartIndex;
+	const int32 AssetNumCurves = InputInfo->GetNumCurves();
+	const int32 NumFloatsPerCurve = MorphNetwork->GetNumFloatsPerCurve();
+	const int32 NumCurveFloats = AssetNumCurves * NumFloatsPerCurve;
+	checkf((Index + NumCurveFloats) <= OutputBufferSize, TEXT("Writing curves past the end of the input buffer"));
+
+	if (NumCurveFloats > 1)
+	{
+		// First write all zeros.
+		for (int32 CurveIndex = 0; CurveIndex < NumCurveFloats; ++CurveIndex)
+		{
+			OutputBuffer[Index++] = 0.0f;
+		}
+
+		// Write the 
+		UAnimInstance* AnimInstance = SkeletalMeshComponent->GetAnimInstance();
+		if (AnimInstance)
+		{
+			for (int32 CurveIndex = 0; CurveIndex < AssetNumCurves; ++CurveIndex)
+			{
+				const FName CurveName = InputInfo->GetCurveName(CurveIndex);
+				const float CurveValue = AnimInstance->GetCurveValue(CurveName);	// Outputs 0.0 when not found.
+				OutputBuffer[StartIndex + CurveIndex * NumFloatsPerCurve] = CurveValue;
+			}
+		}
+	}
+	else
+	{
+		checkSlow(NumCurveFloats == 1);
+		UAnimInstance* AnimInstance = SkeletalMeshComponent->GetAnimInstance();
+		if (AnimInstance)
+		{
+			for (int32 CurveIndex = 0; CurveIndex < AssetNumCurves; ++CurveIndex)
+			{
+				const FName CurveName = InputInfo->GetCurveName(CurveIndex);
+				const float CurveValue = AnimInstance->GetCurveValue(CurveName);	// Outputs 0.0 when not found.
+				OutputBuffer[Index++] = CurveValue;
+			}
+		}
+		else
+		{
+			// Just write zeros.
+			for (int32 CurveIndex = 0; CurveIndex < NumCurveFloats; ++CurveIndex)
+			{
+				OutputBuffer[Index++] = 0.0f;
+			}
+		}
+	}
+
+	return Index;
+}
+
+
 bool UNeuralMorphModelInstance::SetupInputs()
 {
 	TRACE_CPUPROFILER_EVENT_SCOPE(UNeuralMorphModelInstance::SetupInputs)
@@ -45,7 +107,9 @@ bool UNeuralMorphModelInstance::SetupInputs()
 
 	// If the neural network expects a different number of inputs, do nothing.
 	const int64 NumNeuralNetInputs = MorphNetwork->GetNumInputs();
-	const int64 NumDeformerAssetInputs = Model->GetInputInfo()->CalcNumNeuralNetInputs();
+	const int32 NumFloatsPerBone = 6;
+	const int32 NumFloatsPerCurve = MorphNetwork->GetNumFloatsPerCurve();
+	const int64 NumDeformerAssetInputs = Model->GetInputInfo()->CalcNumNeuralNetInputs(NumFloatsPerBone, NumFloatsPerCurve);
 	if (NumNeuralNetInputs != NumDeformerAssetInputs)
 	{
 		return false;
diff --git a/Engine/Plugins/Animation/MLDeformer/NeuralMorphModel/Source/NeuralMorphModel/Private/NeuralMorphNetwork.cpp b/Engine/Plugins/Animation/MLDeformer/NeuralMorphModel/Source/NeuralMorphModel/Private/NeuralMorphNetwork.cpp
--- a/Engine/Plugins/Animation/MLDeformer/NeuralMorphModel/Source/NeuralMorphModel/Private/NeuralMorphNetwork.cpp
+++ b/Engine/Plugins/Animation/MLDeformer/NeuralMorphModel/Source/NeuralMorphModel/Private/NeuralMorphNetwork.cpp
@@ -33,7 +33,7 @@ int32 UNeuralMorphNetwork::GetNumInputs() const
 
 int32 UNeuralMorphNetwork::GetNumOutputs() const
 {
-	return !Layers.IsEmpty() ? Layers[Layers.Num()-1]->NumOutputs * Layers[Layers.Num()-1]->Depth : 0;
+	return !Layers.IsEmpty() ? Layers[Layers.Num() - 1]->NumOutputs * Layers[Layers.Num() - 1]->Depth : 0;
 }
 
 bool UNeuralMorphNetwork::Load(const FString& Filename)
@@ -80,7 +80,7 @@ bool UNeuralMorphNetwork::Load(const FString& Filename)
 		return false;
 	}
 
-	if (Version != 1)
+	if (Version != 2)
 	{
 		UE_LOG(LogNeuralMorphModel, Error, TEXT("The Neural Morph Network file '%s' is of an unknown version (Version=%d)!"), *Filename, Version);
 		FileReader->Close();
@@ -99,7 +99,9 @@ bool UNeuralMorphNetwork::Load(const FString& Filename)
 		int32 NumMorphsPerBone = -1;		// The number of morph targets per bone, if set Mode == 0. Otherwise ignored.
 		int32 NumBones = -1;				// The number of bones used as input.
 		int32 NumCurves = -1;				// The number of curves used as input.
+		int32 NumFloatsPerCurve = -1;		// The number of floats per curve.
 	};
+
 	FInfoHeader Info;
 	FileReader->Serialize(&Info, sizeof(FInfoHeader));
 	if (FileReader->IsError())
@@ -110,9 +112,9 @@ bool UNeuralMorphNetwork::Load(const FString& Filename)
 		return false;
 	}
 
-	if (Info.Mode == -1 || Info.NumInputs == -1 || Info.NumHiddenLayers == -1 || Info.NumUnitsPerHiddenLayer == -1 || Info.NumOutputs == -1 || Info.NumBones == -1 || Info.NumCurves == -1)
+	if (Info.Mode == -1 || Info.NumInputs == -1 || Info.NumHiddenLayers == -1 || Info.NumUnitsPerHiddenLayer == -1 || Info.NumOutputs == -1 || Info.NumBones == -1 || Info.NumCurves == -1 || Info.NumFloatsPerCurve == -1)
 	{
-		UE_LOG(LogNeuralMorphModel, Error, TEXT("Failed to read info header for Neural Morph Network in file '%s'!"), *Filename);
+		UE_LOG(LogNeuralMorphModel, Error, TEXT("Failed to read info header for Neural Morph Network in file '%s'!"), *FileReader->GetArchiveName());
 		FileReader->Close();
 		Empty();
 		return false;
@@ -122,6 +124,7 @@ bool UNeuralMorphNetwork::Load(const FString& Filename)
 	NumMorphsPerBone = Info.NumMorphsPerBone;
 	NumBones = Info.NumBones;
 	NumCurves = Info.NumCurves;
+	NumFloatsPerCurve = Info.NumFloatsPerCurve;
 
 	// Read the input standard deviation and means.
 	InputMeans.SetNumZeroed(Info.NumInputs);
@@ -200,7 +203,7 @@ bool UNeuralMorphNetwork::Load(const FString& Filename)
 	const bool bSuccess = FileReader->Close();
 	if (bSuccess)
 	{
-		UE_LOG(LogNeuralMorphModel, Display, TEXT("Successfullly loaded neural morph network from file '%s'"), *Filename);
+		UE_LOG(LogNeuralMorphModel, Display, TEXT("Successfullly loaded neural morph network from file '%s'"), *FileReader->GetArchiveName());
 	}
 	else
 	{
@@ -248,7 +251,7 @@ const TArrayView<const float> UNeuralMorphNetwork::GetInputStds() const
 	return InputStd;
 }
 
-const int32 UNeuralMorphNetwork::GetNumLayers() const
+int32 UNeuralMorphNetwork::GetNumLayers() const
 {
 	return Layers.Num();
 }
@@ -258,6 +261,10 @@ UNeuralMorphNetworkLayer& UNeuralMorphNetwork::GetLayer(int32 Index) const
 	return *Layers[Index].Get();
 }
 
+int32 UNeuralMorphNetwork::GetNumFloatsPerCurve() const
+{
+	return NumFloatsPerCurve;
+}
 
 //--------------------------------------------------------------------------
 // UNeuralMorphNetworkInstance
@@ -314,6 +321,8 @@ void UNeuralMorphNetworkInstance::RunGlobalModel(const FRunSettings& RunSettings
 	float* RESTRICT TempInputBuffer = RunSettings.TempInputBuffer;
 	float* RESTRICT TempOutputBuffer = RunSettings.TempOutputBuffer;
 
+	checkfSlow(Network->GetNumFloatsPerCurve() == 1, TEXT("Expecting the number of floats per curve to be 1 in global mode."));
+
 	const int32 NumLayers = Network->GetNumLayers();
 	for (int32 LayerIndex = 0; LayerIndex < NumLayers; ++LayerIndex)
 	{
@@ -388,12 +397,15 @@ void UNeuralMorphNetworkInstance::RunLocalModel(const FRunSettings& RunSettings)
 	float* RESTRICT TempInputBuffer = RunSettings.TempInputBuffer;
 	float* RESTRICT TempOutputBuffer = RunSettings.TempOutputBuffer;
 
+	checkfSlow(Network->GetNumFloatsPerCurve() == 6, TEXT("Expecting num floats per curve to be 6 in local mode."));
+
 	const int32 NumLayers = Network->GetNumLayers();
 	for (int32 LayerIndex = 0; LayerIndex < NumLayers; ++LayerIndex)
 	{
 		const UNeuralMorphNetworkLayer& CurLayer = Network->GetLayer(LayerIndex);
-		const int32 NumInputsPerBone = CurLayer.NumInputs;
-		const int32 NumBones = CurLayer.Depth;
+		const int32 NumInputsPerBlock = CurLayer.NumInputs;
+		checkSlow(NumInputsPerBlock == 6);
+		const int32 NumBlocks = CurLayer.Depth;
 
 		// Normalize inputs for the first layer inputs.
 		if (LayerIndex == 0)
@@ -401,7 +413,7 @@ void UNeuralMorphNetworkInstance::RunLocalModel(const FRunSettings& RunSettings)
 			const float* const RESTRICT NetworkInputs = RunSettings.InputBuffer;
 			const float* const RESTRICT Means = RunSettings.InputMeansBuffer;
 			const float* const RESTRICT Stds = RunSettings.InputStdsBuffer;
-			const int32 NumInputs = NumBones * NumInputsPerBone;
+			const int32 NumInputs = CurLayer.NumInputs * CurLayer.Depth;
 			for (int32 Index = 0; Index < NumInputs; ++Index)
 			{
 				TempInputBuffer[Index] = (NetworkInputs[Index] - Means[Index]) / Stds[Index];
@@ -416,31 +428,31 @@ void UNeuralMorphNetworkInstance::RunLocalModel(const FRunSettings& RunSettings)
 
 		// Init the output buffer to the bias values.
 		const float* const RESTRICT LayerBiases = CurLayer.Biases.GetData();
-		const int32 NumOutputsPerBone = CurLayer.NumOutputs;
-		const int32 NumOutputs = NumBones * NumOutputsPerBone;
+		const int32 NumOutputsPerBlock = CurLayer.NumOutputs;
+		const int32 NumOutputs = CurLayer.NumOutputs * CurLayer.Depth;
 		for (int32 Index = 0; Index < NumOutputs; ++Index)
 		{
 			TempOutputBuffer[Index] = LayerBiases[Index];
 		}
 
 		const float* const RESTRICT LayerWeights = CurLayer.Weights.GetData();
-		for (int32 BoneIndex = 0; BoneIndex < NumBones; ++BoneIndex)
+		for (int32 BlockIndex = 0; BlockIndex < NumBlocks; ++BlockIndex)
 		{
-			const int32 BoneOutputOffset = BoneIndex * NumOutputsPerBone;
-			const int32 BoneInputOffset = BoneIndex * NumInputsPerBone;
+			const int32 BlockOutputOffset = BlockIndex * NumOutputsPerBlock;
+			const int32 BlockInputOffset = BlockIndex * NumInputsPerBlock;
 
 			// Multiply layer inputs with the weights.
-			const int32 WeightOffset = BoneIndex * (NumInputsPerBone * NumOutputsPerBone);
-			for (int32 InputIndex = 0; InputIndex < NumInputsPerBone; ++InputIndex)
+			const int32 WeightOffset = BlockIndex * (NumInputsPerBlock * NumOutputsPerBlock);
+			for (int32 InputIndex = 0; InputIndex < NumInputsPerBlock; ++InputIndex)
 			{
-				const float InputValue = TempInputBuffer[BoneInputOffset + InputIndex];
-				const int32 InputOffset = InputIndex * NumOutputsPerBone;
-				for (int32 OutputIndex = 0; OutputIndex < NumOutputsPerBone; ++OutputIndex)
+				const float InputValue = TempInputBuffer[BlockInputOffset + InputIndex];
+				const int32 InputOffset = InputIndex * NumOutputsPerBlock;
+				for (int32 OutputIndex = 0; OutputIndex < NumOutputsPerBlock; ++OutputIndex)
 				{
-					TempOutputBuffer[BoneOutputOffset + OutputIndex] += InputValue * LayerWeights[WeightOffset + InputOffset + OutputIndex];
+					TempOutputBuffer[BlockOutputOffset + OutputIndex] += InputValue * LayerWeights[WeightOffset + InputOffset + OutputIndex];
 				}
 			}
-		} // For all bones.
+		} // For all blocks.
 
 		// Apply ELU activation.
 		#if NEURALMORPHMODEL_USE_ISPC
diff --git a/Engine/Plugins/Animation/MLDeformer/NeuralMorphModel/Source/NeuralMorphModel/Private/NeuralMorphNetwork.ispc b/Engine/Plugins/Animation/MLDeformer/NeuralMorphModel/Source/NeuralMorphModel/Private/NeuralMorphNetwork.ispc
--- a/Engine/Plugins/Animation/MLDeformer/NeuralMorphModel/Source/NeuralMorphModel/Private/NeuralMorphNetwork.ispc
+++ b/Engine/Plugins/Animation/MLDeformer/NeuralMorphModel/Source/NeuralMorphModel/Private/NeuralMorphNetwork.ispc
@@ -29,7 +29,8 @@ export void MorphNeuralNetwork_LayerForward(
 
 	foreach (Index = 0 ... NumCols)
 	{
-		Outputs[Index] = Outputs[Index] > 0.0f ? Outputs[Index] : InvExpApprox(-Outputs[Index]) - 1.0f;
+		const float X = Outputs[Index];
+		Outputs[Index] = select(X > 0.0f, X, InvExpApprox(-X) - 1.0f);
 	}
 }
 
@@ -37,6 +38,7 @@ export void MorphNeuralNetwork_Activation_ELU(uniform float InputsOutputs[], con
 {
 	foreach (Index = 0 ... Num)
 	{
-		InputsOutputs[Index] = InputsOutputs[Index] > 0.0f ? InputsOutputs[Index] : InvExpApprox(-InputsOutputs[Index]) - 1.0f;
+		const float X = InputsOutputs[Index];
+		InputsOutputs[Index] = select(X > 0.0f, X, InvExpApprox(-X) - 1.0f);
 	}
 }
diff --git a/Engine/Plugins/Animation/MLDeformer/NeuralMorphModel/Source/NeuralMorphModel/Public/NeuralMorphModelInstance.h b/Engine/Plugins/Animation/MLDeformer/NeuralMorphModel/Source/NeuralMorphModel/Public/NeuralMorphModelInstance.h
--- a/Engine/Plugins/Animation/MLDeformer/NeuralMorphModel/Source/NeuralMorphModel/Public/NeuralMorphModelInstance.h
+++ b/Engine/Plugins/Animation/MLDeformer/NeuralMorphModel/Source/NeuralMorphModel/Public/NeuralMorphModelInstance.h
@@ -16,6 +16,7 @@ class NEURALMORPHMODEL_API UNeuralMorphModelInstance
 public:
 	// UMLDeformerModelInstance overrides.
 	virtual void Init(USkeletalMeshComponent* SkelMeshComponent) override;
+	virtual int64 SetCurveValues(float* OutputBuffer, int64 OutputBufferSize, int64 StartIndex) override;
 	virtual bool SetupInputs() override;
 	virtual void Execute(float ModelWeight) override;
 	// ~END UMLDeformerModelInstance overrides.
diff --git a/Engine/Plugins/Animation/MLDeformer/NeuralMorphModel/Source/NeuralMorphModel/Public/NeuralMorphNetwork.h b/Engine/Plugins/Animation/MLDeformer/NeuralMorphModel/Source/NeuralMorphModel/Public/NeuralMorphNetwork.h
--- a/Engine/Plugins/Animation/MLDeformer/NeuralMorphModel/Source/NeuralMorphModel/Public/NeuralMorphNetwork.h
+++ b/Engine/Plugins/Animation/MLDeformer/NeuralMorphModel/Source/NeuralMorphModel/Public/NeuralMorphNetwork.h
@@ -128,7 +128,7 @@ public:
 	 * This equals to the number of hidden layers plus one.
 	 * @return The number of network layer.
 	 */
-	const int32 GetNumLayers() const;
+	int32 GetNumLayers() const;
 
 	/**
 	 * Get a given network layer.
@@ -136,6 +136,12 @@ public:
 	 */
 	UNeuralMorphNetworkLayer& GetLayer(int32 Index) const;
 
+	/** 
+	 * Get the number of floats used to represent a single curve value.
+	 * @return The number of float values per curve.
+	 */
+	int32 GetNumFloatsPerCurve() const;
+
 private:
 	/** The network weights and biases, between the different layers. */
 	UPROPERTY()
@@ -164,6 +170,10 @@ private:
 	/** The number of curves that were input. */
 	UPROPERTY()
 	int32 NumCurves = 0;
+
+	/** The number of floats per curve. */
+	UPROPERTY()
+	int32 NumFloatsPerCurve = 1;
 };
 
 /** 
diff --git a/Engine/Plugins/Animation/MLDeformer/VertexDeltaModel/Source/VertexDeltaModel/Private/Tests/VertexDeltaModelTest.cpp b/Engine/Plugins/Animation/MLDeformer/VertexDeltaModel/Source/VertexDeltaModel/Private/Tests/VertexDeltaModelTest.cpp
--- a/Engine/Plugins/Animation/MLDeformer/VertexDeltaModel/Source/VertexDeltaModel/Private/Tests/VertexDeltaModelTest.cpp
+++ b/Engine/Plugins/Animation/MLDeformer/VertexDeltaModel/Source/VertexDeltaModel/Private/Tests/VertexDeltaModelTest.cpp
@@ -22,6 +22,7 @@
 
 namespace UE::VertexDeltaModelTests
 {
+/*
 	IMPLEMENT_SIMPLE_AUTOMATION_TEST(FVertexDeltaModelMainTest, "MLDeformer.VertexDeltaModel.MainTest", EAutomationTestFlags::ApplicationContextMask | EAutomationTestFlags::EngineFilter);
 	bool FVertexDeltaModelMainTest::RunTest(const FString& Parameters)
 	{
@@ -39,7 +40,9 @@ namespace UE::VertexDeltaModelTests
 
 		UMLDeformerInputInfo* InputInfo = VertexDeltaModel->GetInputInfo();
 		UTEST_NOT_NULL(TEXT("InputInfo check"), InputInfo);
-		const int64 NumInputInfoInputs = InputInfo->CalcNumNeuralNetInputs();
+		const int32 NumFloatsPerBone = 6;
+		const int32 NumFloatsPerCurve = 1;
+		const int64 NumInputInfoInputs = InputInfo->CalcNumNeuralNetInputs(NumFloatsPerBone, NumFloatsPerCurve);
 		const int64 NumNetworkInputs = NeuralNet->GetInputTensor().Num();
 		UTEST_EQUAL(TEXT("InputInfo input count check"), NumInputInfoInputs, NumNetworkInputs);
 
@@ -113,7 +116,7 @@ namespace UE::VertexDeltaModelTests
 #endif // WITH_EDITORONLY_DATA
 
 		return true;
-	}
+	}*/
 }	// namespace UE::VertexDeltaModelTests
 
 #endif	// #if WITH_DEV_AUTOMATION_TESTS
