[MLDeformer] Fix a crash in the neural morph model.

#trivial

diff --git a/Engine/Plugins/Animation/MLDeformer/NeuralMorphModel/Source/NeuralMorphModel/Private/NeuralMorphModelInstance.cpp b/Engine/Plugins/Animation/MLDeformer/NeuralMorphModel/Source/NeuralMorphModel/Private/NeuralMorphModelInstance.cpp
--- a/Engine/Plugins/Animation/MLDeformer/NeuralMorphModel/Source/NeuralMorphModel/Private/NeuralMorphModelInstance.cpp
+++ b/Engine/Plugins/Animation/MLDeformer/NeuralMorphModel/Source/NeuralMorphModel/Private/NeuralMorphModelInstance.cpp
@@ -31,7 +31,7 @@ int64 UNeuralMorphModelInstance::SetCurveValues(float* OutputBuffer, int64 Outpu
 
 	int64 Index = StartIndex;
 	const int32 AssetNumCurves = InputInfo->GetNumCurves();
-	const int32 NumFloatsPerCurve = MorphNetwork->GetNumFloatsPerCurve();
+	const int32 NumFloatsPerCurve = MorphNetwork ? MorphNetwork->GetNumFloatsPerCurve() : 1;
 	const int32 NumCurveFloats = AssetNumCurves * NumFloatsPerCurve;
 	checkf((Index + NumCurveFloats) <= OutputBufferSize, TEXT("Writing curves past the end of the input buffer"));
 
