[MLDeformer]
- Add support for joint and curve groups in the local model.
- Moved the Include Normals option in the morph target category out of the compression group.
- Fix some duplicate log messages.

diff --git a/Engine/Plugins/Animation/MLDeformer/MLDeformerFramework/Source/MLDeformerFramework/Private/MLDeformerComponent.cpp b/Engine/Plugins/Animation/MLDeformer/MLDeformerFramework/Source/MLDeformerFramework/Private/MLDeformerComponent.cpp
--- a/Engine/Plugins/Animation/MLDeformer/MLDeformerFramework/Source/MLDeformerFramework/Private/MLDeformerComponent.cpp
+++ b/Engine/Plugins/Animation/MLDeformer/MLDeformerFramework/Source/MLDeformerFramework/Private/MLDeformerComponent.cpp
@@ -39,6 +39,7 @@ void UMLDeformerComponent::Init()
 		if (ModelInstance)
 		{
 			ModelInstance->Release();
+			ModelInstance->ConditionalBeginDestroy();
 			ModelInstance = nullptr;
 		}
 		return;
@@ -51,6 +52,7 @@ void UMLDeformerComponent::Init()
 		if (ModelInstance)
 		{
 			ModelInstance->Release();
+			ModelInstance->ConditionalBeginDestroy();
 			ModelInstance = nullptr;
 		}
 		ModelInstance = Model->CreateModelInstance(this);
@@ -231,9 +233,9 @@ void UMLDeformerComponent::TickComponent(float DeltaTime, enum ELevelTick TickTy
 		{
 			TRACE_CPUPROFILER_EVENT_SCOPE(UMLDeformerComponent::TickComponent)
 			float ApplyWeight = Weight;
-			if (GMLDeformerOverrideWeight >= 0.f)
+			if (GMLDeformerOverrideWeight >= 0.0f)
 			{
-				ApplyWeight = FMath::Min(GMLDeformerOverrideWeight, 1.f);
+				ApplyWeight = FMath::Min(GMLDeformerOverrideWeight, 1.0f);
 			}
 			ModelInstance->Tick(DeltaTime, ApplyWeight);
 
diff --git a/Engine/Plugins/Animation/MLDeformer/MLDeformerFramework/Source/MLDeformerFramework/Private/MLDeformerGeomCacheHelpers.cpp b/Engine/Plugins/Animation/MLDeformer/MLDeformerFramework/Source/MLDeformerFramework/Private/MLDeformerGeomCacheHelpers.cpp
--- a/Engine/Plugins/Animation/MLDeformer/MLDeformerFramework/Source/MLDeformerFramework/Private/MLDeformerGeomCacheHelpers.cpp
+++ b/Engine/Plugins/Animation/MLDeformer/MLDeformerFramework/Source/MLDeformerFramework/Private/MLDeformerGeomCacheHelpers.cpp
@@ -76,7 +76,8 @@ namespace UE::MLDeformer
 			TArray<FMLDeformerGeomCacheMeshMapping> MeshMappings;
 			TArray<FString> FailedNames;
 			TArray<FString> VertexMisMatchNames;
-			GenerateGeomCacheMeshMappings(InSkelMesh, InGeomCache, MeshMappings, FailedNames, VertexMisMatchNames);
+			const bool bSuppressLogMessages = true;
+			GenerateGeomCacheMeshMappings(InSkelMesh, InGeomCache, MeshMappings, FailedNames, VertexMisMatchNames, bSuppressLogMessages);
 
 			// List all mesh names that have issues.
 			FString ErrorString;
@@ -125,7 +126,7 @@ namespace UE::MLDeformer
 		return (TrackName.Find(MeshName) == 0);
 	}
 
-	void GenerateGeomCacheMeshMappings(USkeletalMesh* SkelMesh, UGeometryCache* GeomCache, TArray<FMLDeformerGeomCacheMeshMapping>& OutMeshMappings, TArray<FString>& OutFailedImportedMeshNames, TArray<FString>& OutVertexMisMatchNames)
+	void GenerateGeomCacheMeshMappings(USkeletalMesh* SkelMesh, UGeometryCache* GeomCache, TArray<FMLDeformerGeomCacheMeshMapping>& OutMeshMappings, TArray<FString>& OutFailedImportedMeshNames, TArray<FString>& OutVertexMisMatchNames, bool bSuppressLog)
 	{
 		OutMeshMappings.Empty();
 		OutFailedImportedMeshNames.Empty();
@@ -217,7 +218,10 @@ namespace UE::MLDeformer
 
 					// We found a match, no need to iterate over more Tracks.
 					bFoundMatch = true;
-					UE_LOG(LogMLDeformer, Display, TEXT("Mapped geom cache track '%s' to SkelMesh submesh '%s'"), *Track->GetName(), *SkelMeshName);
+					if (!bSuppressLog)
+					{
+						UE_LOG(LogMLDeformer, Display, TEXT("Mapped geom cache track '%s' to SkelMesh submesh '%s'"), *Track->GetName(), *SkelMeshName);
+					}
 					break;
 				} // If the track name matches the skeletal meshes internal mesh name.
 			} // For all meshes in the Skeletal Mesh.
@@ -225,7 +229,10 @@ namespace UE::MLDeformer
 			if (Track && !bFoundMatch)
 			{
 				OutFailedImportedMeshNames.Add(Track->GetName());
-				UE_LOG(LogMLDeformer, Warning, TEXT("Geometry cache '%s' cannot be matched with a mesh inside the Skeletal Mesh."), *Track->GetName());
+				if (!bSuppressLog)
+				{
+					UE_LOG(LogMLDeformer, Warning, TEXT("Geometry cache '%s' cannot be matched with a mesh inside the Skeletal Mesh."), *Track->GetName());
+				}
 			}
 		} // For all tracks.
 	}
diff --git a/Engine/Plugins/Animation/MLDeformer/MLDeformerFramework/Source/MLDeformerFramework/Private/MLDeformerInputInfo.cpp b/Engine/Plugins/Animation/MLDeformer/MLDeformerFramework/Source/MLDeformerFramework/Private/MLDeformerInputInfo.cpp
--- a/Engine/Plugins/Animation/MLDeformer/MLDeformerFramework/Source/MLDeformerFramework/Private/MLDeformerInputInfo.cpp
+++ b/Engine/Plugins/Animation/MLDeformer/MLDeformerFramework/Source/MLDeformerFramework/Private/MLDeformerInputInfo.cpp
@@ -28,6 +28,15 @@ void UMLDeformerInputInfo::OnPostLoad()
 	PRAGMA_ENABLE_DEPRECATION_WARNINGS
 }
 
+void UMLDeformerInputInfo::CopyMembersFrom(UMLDeformerInputInfo* Other)
+{
+	BoneNames = Other->BoneNames;
+	CurveNames = Other->CurveNames;
+	NumBaseMeshVertices = Other->NumBaseMeshVertices;
+	NumTargetMeshVertices = Other->NumTargetMeshVertices;
+	SkeletalMesh = Other->SkeletalMesh;
+}
+
 void UMLDeformerInputInfo::UpdateNameStrings()
 {
 #if WITH_EDITORONLY_DATA
diff --git a/Engine/Plugins/Animation/MLDeformer/MLDeformerFramework/Source/MLDeformerFramework/Private/MLDeformerMorphModelInstance.cpp b/Engine/Plugins/Animation/MLDeformer/MLDeformerFramework/Source/MLDeformerFramework/Private/MLDeformerMorphModelInstance.cpp
--- a/Engine/Plugins/Animation/MLDeformer/MLDeformerFramework/Source/MLDeformerFramework/Private/MLDeformerMorphModelInstance.cpp
+++ b/Engine/Plugins/Animation/MLDeformer/MLDeformerFramework/Source/MLDeformerFramework/Private/MLDeformerMorphModelInstance.cpp
@@ -138,9 +138,6 @@ void UMLDeformerMorphModelInstance::Tick(float DeltaTime, float ModelWeight)
 
 FExternalMorphSetWeights* UMLDeformerMorphModelInstance::FindWeightData(int32 LOD) const
 {
-	const UMLDeformerMorphModel* MorphModel = Cast<UMLDeformerMorphModel>(Model);
-	checkSlow(MorphModel);
-
 	// Check if our LOD index is valid first, as we might not have registered yet.
 	USkeletalMeshComponent* SkelMeshComponent = SkeletalMeshComponent.Get();
 	if (SkelMeshComponent == nullptr || !SkelMeshComponent->IsValidExternalMorphSetLODIndex(LOD))
@@ -155,8 +152,6 @@ FExternalMorphSetWeights* UMLDeformerMorphModelInstance::FindWeightData(int32 LO
 
 void UMLDeformerMorphModelInstance::HandleZeroModelWeight()
 {
-	TRACE_CPUPROFILER_EVENT_SCOPE(UNeuralMorphModelInstance::HandleZeroModelWeight)
-
 	const int LOD = 0;	// For now we only support LOD 0, as we can't setup an ML Deformer per LOD yet.
 	FExternalMorphSetWeights* WeightData = FindWeightData(LOD);
 	if (WeightData)
diff --git a/Engine/Plugins/Animation/MLDeformer/MLDeformerFramework/Source/MLDeformerFramework/Public/MLDeformerGeomCacheHelpers.h b/Engine/Plugins/Animation/MLDeformer/MLDeformerFramework/Source/MLDeformerFramework/Public/MLDeformerGeomCacheHelpers.h
--- a/Engine/Plugins/Animation/MLDeformer/MLDeformerFramework/Source/MLDeformerFramework/Public/MLDeformerGeomCacheHelpers.h
+++ b/Engine/Plugins/Animation/MLDeformer/MLDeformerFramework/Source/MLDeformerFramework/Public/MLDeformerGeomCacheHelpers.h
@@ -69,8 +69,9 @@ namespace UE::MLDeformer
 	 * @param OutMeshMapping The mapping that will be generated. This information maps each geometry cache map with a given compatible mesh inside a skeletal mesh.
 	 * @param OutFailedImportedMeshNames The array of geometry cache track names for which no mesh could be found inside the skeletal mesh object.
 	 * @param OutVertexMisMatchNames The array of geometry cache track names for which the vertex counts didn't match the related meshes inside the skeletal mesh.	 
+	 * @param bSuppressLog Set to true to suppress any log messages.
 	 */
-	MLDEFORMERFRAMEWORK_API void GenerateGeomCacheMeshMappings(USkeletalMesh* SkelMesh, UGeometryCache* GeomCache, TArray<FMLDeformerGeomCacheMeshMapping>& OutMeshMappings, TArray<FString>& OutFailedImportedMeshNames, TArray<FString>& OutVertexMisMatchNames);
+	MLDEFORMERFRAMEWORK_API void GenerateGeomCacheMeshMappings(USkeletalMesh* SkelMesh, UGeometryCache* GeomCache, TArray<FMLDeformerGeomCacheMeshMapping>& OutMeshMappings, TArray<FString>& OutFailedImportedMeshNames, TArray<FString>& OutVertexMisMatchNames, bool bSuppressLog = false);
 
 	/**
 	 * Sample the vertex position data of a geometry cache, at a given time stamp.
diff --git a/Engine/Plugins/Animation/MLDeformer/MLDeformerFramework/Source/MLDeformerFramework/Public/MLDeformerInputInfo.h b/Engine/Plugins/Animation/MLDeformer/MLDeformerFramework/Source/MLDeformerFramework/Public/MLDeformerInputInfo.h
--- a/Engine/Plugins/Animation/MLDeformer/MLDeformerFramework/Source/MLDeformerFramework/Public/MLDeformerInputInfo.h
+++ b/Engine/Plugins/Animation/MLDeformer/MLDeformerFramework/Source/MLDeformerFramework/Public/MLDeformerInputInfo.h
@@ -74,6 +74,12 @@ public:
 	 */
 	virtual int32 CalcNumNeuralNetInputs(int32 NumFloatsPerBone, int32 NumFloatsPerCurve) const;
 
+	/**
+	 * Copy members from the other class.
+	 * @param Other The other class to copy over members from.
+	 **/
+	virtual void CopyMembersFrom(UMLDeformerInputInfo* Other);
+
 	/**
 	 * Get the number of bones that we trained on.
 	 * @return The number of bones.
diff --git a/Engine/Plugins/Animation/MLDeformer/MLDeformerFramework/Source/MLDeformerFrameworkEditor/Private/MLDeformerEditorModel.cpp b/Engine/Plugins/Animation/MLDeformer/MLDeformerFramework/Source/MLDeformerFrameworkEditor/Private/MLDeformerEditorModel.cpp
--- a/Engine/Plugins/Animation/MLDeformer/MLDeformerFramework/Source/MLDeformerFrameworkEditor/Private/MLDeformerEditorModel.cpp
+++ b/Engine/Plugins/Animation/MLDeformer/MLDeformerFramework/Source/MLDeformerFrameworkEditor/Private/MLDeformerEditorModel.cpp
@@ -124,6 +124,10 @@ namespace UE::MLDeformer
 		InPersonaPreviewScene->AddComponent(SkelMeshComponent, FTransform::Identity);
 		InPersonaPreviewScene->SetAdditionalMeshesSelectable(false);
 		InPersonaPreviewScene->SetPreviewMesh(Mesh);
+		if (SkelMeshComponent->GetAnimInstance())
+		{
+			SkelMeshComponent->GetAnimInstance()->GetRequiredBones().SetUseRAWData(true);
+		}
 
 		// Register the editor actor.
 		const FLinearColor LabelColor = FMLDeformerEditorStyle::Get().GetColor("MLDeformer.BaseMesh.LabelColor");
@@ -155,6 +159,7 @@ namespace UE::MLDeformer
 		SkelMeshComponent->SetVisibility(false);
 		SkelMeshComponent->MarkRenderStateDirty();
 
+
 		// Register the editor actor.
 		const FLinearColor LabelColor = FMLDeformerEditorStyle::Get().GetColor("MLDeformer.BaseMesh.LabelColor");
 		FMLDeformerEditorActor::FConstructSettings Settings;
@@ -306,23 +311,33 @@ namespace UE::MLDeformer
 
 	void FMLDeformerEditorModel::Tick(FEditorViewportClient* ViewportClient, float DeltaTime)
 	{
-		// Force the training sequence to use Step interpolation and sample raw animation data.
+		// Force the training sequence to use Step interpolation.
 		// We do this in the Tick, as this is reset to false in the engine sometimes.
 		UAnimSequence* TrainingAnimSequence = Model->GetAnimSequence();
 		if (TrainingAnimSequence)
 		{
-			TrainingAnimSequence->bUseRawDataOnly = true;
 			TrainingAnimSequence->Interpolation = EAnimInterpolationType::Step;
 		}
 
+		// Use raw data for everything.
+		for (FMLDeformerEditorActor* Actor : EditorActors)
+		{
+			if (Actor && Actor->GetSkeletalMeshComponent())
+			{
+				USkeletalMeshComponent* SkelMeshComp = Actor->GetSkeletalMeshComponent();
+				if (SkelMeshComp->GetAnimInstance())
+				{
+					SkelMeshComp->GetAnimInstance()->GetRequiredBones().SetUseRAWData(true);
+				}
+			}
+		}
+
 		// Do the same for the test anim sequence.
 		UMLDeformerVizSettings* VizSettings = Model->GetVizSettings();
 		{
 			UAnimSequence* TestAnimSequence = VizSettings->GetTestAnimSequence();
 			if (TestAnimSequence)
-			{
-				TestAnimSequence->bUseRawDataOnly = true;
-
+			{			
 				// Enable step interpolation when doing showing a heatmap vs ground truth.
 				if (VizSettings->HasTestGroundTruth() && VizSettings->GetShowHeatMap() && VizSettings->GetHeatMapMode() == EMLDeformerHeatMapMode::GroundTruth)
 				{
@@ -445,7 +460,6 @@ namespace UE::MLDeformer
 		UAnimSequence* TrainingAnimSequence = Model->GetAnimSequence();
 		if (TrainingAnimSequence)
 		{
-			TrainingAnimSequence->bUseRawDataOnly = true;
 			TrainingAnimSequence->Interpolation = EAnimInterpolationType::Step;
 		}
 
@@ -467,6 +481,10 @@ namespace UE::MLDeformer
 			SkeletalMeshComponent->SetPosition(CurrentPlayTime);
 			SkeletalMeshComponent->SetPlayRate(TestAnimSpeed);
 			SkeletalMeshComponent->Play(true);
+			if (SkeletalMeshComponent->GetAnimInstance())
+			{
+				SkeletalMeshComponent->GetAnimInstance()->GetRequiredBones().SetUseRAWData(true);
+			}
 		}
 
 		// Update the test base model.
@@ -643,7 +661,6 @@ namespace UE::MLDeformer
 			if (PropertyChangedEvent.ChangeType == EPropertyChangeType::ValueSet)
 			{
 				SetResamplingInputOutputsNeeded(true);
-				UpdateEditorInputInfo();
 				UpdateIsReadyForTrainingState();
 				GetEditor()->GetModelDetailsView()->ForceRefresh();
 			}
diff --git a/Engine/Plugins/Animation/MLDeformer/MLDeformerFramework/Source/MLDeformerFrameworkEditor/Private/MLDeformerEditorModule.cpp b/Engine/Plugins/Animation/MLDeformer/MLDeformerFramework/Source/MLDeformerFrameworkEditor/Private/MLDeformerEditorModule.cpp
--- a/Engine/Plugins/Animation/MLDeformer/MLDeformerFramework/Source/MLDeformerFrameworkEditor/Private/MLDeformerEditorModule.cpp
+++ b/Engine/Plugins/Animation/MLDeformer/MLDeformerFramework/Source/MLDeformerFrameworkEditor/Private/MLDeformerEditorModule.cpp
@@ -26,7 +26,7 @@ namespace UE::MLDeformer
 
 		// Register object detail customizations.
 		FPropertyEditorModule& PropertyModule = FModuleManager::GetModuleChecked<FPropertyEditorModule>("PropertyEditor");
-		PropertyModule.RegisterCustomPropertyTypeLayout("CurveReference", FOnGetPropertyTypeCustomizationInstance::CreateStatic(&FMLDeformerCurveReferenceCustomization::MakeInstance) );
+		PropertyModule.RegisterCustomPropertyTypeLayout("MLDeformerCurveReference", FOnGetPropertyTypeCustomizationInstance::CreateStatic(&FMLDeformerCurveReferenceCustomization::MakeInstance) );
 		PropertyModule.RegisterCustomPropertyTypeLayout("MLDeformerMorphModelQualityLevel", FOnGetPropertyTypeCustomizationInstance::CreateStatic(&FMLDeformerMorphModelQualityLevelCustomization::MakeInstance) );
 		PropertyModule.NotifyCustomizationModuleChanged();
 	}
@@ -47,7 +47,7 @@ namespace UE::MLDeformer
 		if (FModuleManager::Get().IsModuleLoaded(TEXT("PropertyEditor")))
 		{
 			FPropertyEditorModule& PropertyModule = FModuleManager::GetModuleChecked<FPropertyEditorModule>(TEXT("PropertyEditor"));
-			PropertyModule.UnregisterCustomPropertyTypeLayout(TEXT("CurveReference"));
+			PropertyModule.UnregisterCustomPropertyTypeLayout(TEXT("MLDeformerCurveReference"));
 			PropertyModule.UnregisterCustomPropertyTypeLayout(TEXT("MLDeformerMorphModelQualityLevel"));
 			PropertyModule.NotifyCustomizationModuleChanged();
 		}
diff --git a/Engine/Plugins/Animation/MLDeformer/MLDeformerFramework/Source/MLDeformerFrameworkEditor/Private/MLDeformerMorphModelDetails.cpp b/Engine/Plugins/Animation/MLDeformer/MLDeformerFramework/Source/MLDeformerFrameworkEditor/Private/MLDeformerMorphModelDetails.cpp
--- a/Engine/Plugins/Animation/MLDeformer/MLDeformerFramework/Source/MLDeformerFrameworkEditor/Private/MLDeformerMorphModelDetails.cpp
+++ b/Engine/Plugins/Animation/MLDeformer/MLDeformerFramework/Source/MLDeformerFrameworkEditor/Private/MLDeformerMorphModelDetails.cpp
@@ -37,10 +37,11 @@ namespace UE::MLDeformer
 		// Create all the detail categories and add the properties of the base class.
 		FMLDeformerGeomCacheModelDetails::CustomizeDetails(DetailBuilder);
 
+		MorphTargetCategoryBuilder->AddProperty(DetailLayoutBuilder->GetProperty(UMLDeformerMorphModel::GetIncludeMorphTargetNormalsPropertyName(), UMLDeformerMorphModel::StaticClass()));
+
 		IDetailGroup& CompressionGroup = MorphTargetCategoryBuilder->AddGroup("Compression", LOCTEXT("MorphCompressionGroupLabel", "Compression"), false, true);
 		CompressionGroup.AddPropertyRow(DetailLayoutBuilder->GetProperty(UMLDeformerMorphModel::GetMorphDeltaZeroThresholdPropertyName(), UMLDeformerMorphModel::StaticClass()));
 		CompressionGroup.AddPropertyRow(DetailLayoutBuilder->GetProperty(UMLDeformerMorphModel::GetMorphCompressionLevelPropertyName(), UMLDeformerMorphModel::StaticClass()));
-		CompressionGroup.AddPropertyRow(DetailLayoutBuilder->GetProperty(UMLDeformerMorphModel::GetIncludeMorphTargetNormalsPropertyName(), UMLDeformerMorphModel::StaticClass()));
 
 		IDetailGroup& MaskGroup = MorphTargetCategoryBuilder->AddGroup("Mask", LOCTEXT("MorphMaskGroupLabel", "Masking"), false, false);
 		MaskGroup.AddPropertyRow(DetailLayoutBuilder->GetProperty(UMLDeformerMorphModel::GetMaskChannelPropertyName(), UMLDeformerMorphModel::StaticClass()));
diff --git a/Engine/Plugins/Animation/MLDeformer/MLDeformerFramework/Source/MLDeformerFrameworkEditor/Private/MLDeformerSampler.cpp b/Engine/Plugins/Animation/MLDeformer/MLDeformerFramework/Source/MLDeformerFrameworkEditor/Private/MLDeformerSampler.cpp
--- a/Engine/Plugins/Animation/MLDeformer/MLDeformerFramework/Source/MLDeformerFrameworkEditor/Private/MLDeformerSampler.cpp
+++ b/Engine/Plugins/Animation/MLDeformer/MLDeformerFramework/Source/MLDeformerFrameworkEditor/Private/MLDeformerSampler.cpp
@@ -8,6 +8,8 @@
 
 #include "Animation/DebugSkelMeshComponent.h"
 #include "Animation/AnimSequence.h"
+#include "Animation/AnimInstance.h"
+#include "BoneContainer.h"
 #include "BoneWeights.h"
 #include "Engine/SkeletalMesh.h"
 #include "Rendering/SkeletalMeshModel.h"
@@ -73,6 +75,10 @@ namespace UE::MLDeformer
 			SkeletalMeshComponent->bPauseAnims = true;
 			SkeletalMeshComponent->RefreshBoneTransforms();
 			SkeletalMeshComponent->CacheRefToLocalMatrices(BoneMatrices);
+			if (SkeletalMeshComponent->GetAnimInstance())
+			{
+				SkeletalMeshComponent->GetAnimInstance()->GetRequiredBones().SetUseRAWData(true);
+			}
 		}
 	}
 
@@ -254,6 +260,10 @@ namespace UE::MLDeformer
 		SkeletalMeshComponent->Play(false);
 		SkeletalMeshComponent->SetVisibility(false);
 		SkeletalMeshComponent->RefreshBoneTransforms();
+		if (SkeletalMeshComponent->GetAnimInstance())
+		{
+			SkeletalMeshComponent->GetAnimInstance()->GetRequiredBones().SetUseRAWData(true);
+		}
 
 		if (TargetMeshActor.Get() == nullptr)
 		{
diff --git a/Engine/Plugins/Animation/MLDeformer/NeuralMorphModel/Source/NeuralMorphModel/Private/NeuralMorphInputInfo.cpp b/Engine/Plugins/Animation/MLDeformer/NeuralMorphModel/Source/NeuralMorphModel/Private/NeuralMorphInputInfo.cpp
new file mode 100644
--- /dev/null
+++ b/Engine/Plugins/Animation/MLDeformer/NeuralMorphModel/Source/NeuralMorphModel/Private/NeuralMorphInputInfo.cpp
@@ -0,0 +1,183 @@
+// Copyright Epic Games, Inc. All Rights Reserved.
+
+#include "NeuralMorphInputInfo.h"
+#include "NeuralMorphTypes.h"
+#include "MLDeformerCurveReference.h"
+#include "BoneContainer.h"
+
+void UNeuralMorphInputInfo::Reset()
+{
+	Super::Reset();
+	BoneGroups.Empty();
+	CurveGroups.Empty();
+}
+
+TArray<FNeuralMorphBoneGroup>& UNeuralMorphInputInfo::GetBoneGroups()
+{
+	return BoneGroups;
+}
+
+TArray<FNeuralMorphCurveGroup>& UNeuralMorphInputInfo::GetCurveGroups()
+{
+	return CurveGroups;
+}
+
+int32 UNeuralMorphInputInfo::CalcMaxNumGroupItems() const
+{
+	int32 MaxGroupSize = 0;
+
+	const int32 NumBoneGroups = BoneGroups.Num();
+	for (int32 GroupIndex = 0; GroupIndex < NumBoneGroups; ++GroupIndex)
+	{
+		const int32 NumItemsInGroup = BoneGroups[GroupIndex].BoneNames.Num();
+		MaxGroupSize = FMath::Max<int32>(MaxGroupSize, NumItemsInGroup);
+	}
+
+	const int32 NumCurveGroups = CurveGroups.Num();
+	for (int32 GroupIndex = 0; GroupIndex < NumCurveGroups; ++GroupIndex)
+	{
+		const int32 NumItemsInGroup = CurveGroups[GroupIndex].CurveNames.Num();
+		MaxGroupSize = FMath::Max<int32>(MaxGroupSize, NumItemsInGroup);
+	}
+
+	return MaxGroupSize;
+}
+
+void UNeuralMorphInputInfo::GenerateBoneGroupIndices(TArray<int32>& OutBoneGroupIndices)
+{
+	const int32 MaxGroupSize = CalcMaxNumGroupItems();
+
+	// Make sure we have the right number of groups.
+	const int32 NumBoneGroups = BoneGroups.Num();
+	if (OutBoneGroupIndices.Num() != NumBoneGroups * MaxGroupSize)
+	{
+		OutBoneGroupIndices.Reset();
+		OutBoneGroupIndices.AddDefaulted(NumBoneGroups * MaxGroupSize);
+	}
+
+	// Build a list of ints for every group.
+	// The integers in the group are the bone indices we pass transforms for to the inputs.
+	// 0 means the first bone we pass to the network inputs.
+	// 9 means the tenth bone we pass to the network inputs, etc.
+	// All groups will have the same number of items. We pad them with INDEX_NONE when they have less items
+	// compared to the maximum number of items of all groups.
+	for (int32 BoneGroupIndex = 0; BoneGroupIndex < NumBoneGroups; ++BoneGroupIndex)
+	{
+		int32 LastValidIndex = INDEX_NONE;
+		const int32 NumItemsInGroup = BoneGroups[BoneGroupIndex].BoneNames.Num();
+		for (int32 Index = 0; Index < MaxGroupSize; Index++)
+		{
+			if (Index < NumItemsInGroup)
+			{
+				const FName BoneName = BoneGroups[BoneGroupIndex].BoneNames[Index].BoneName;
+				const int32 BoneInputIndex = BoneName.IsValid() ? BoneNames.Find(BoneName) : INDEX_NONE;
+				if (BoneInputIndex != INDEX_NONE)
+				{
+					LastValidIndex = BoneInputIndex;
+				}
+				OutBoneGroupIndices[BoneGroupIndex * MaxGroupSize + Index] = BoneInputIndex;
+			}
+			else
+			{
+				OutBoneGroupIndices[BoneGroupIndex * MaxGroupSize + Index] = LastValidIndex;
+			}
+		}
+	}
+}
+
+void UNeuralMorphInputInfo::GenerateCurveGroupIndices(TArray<int32>& OutCurveGroupIndices)
+{
+	const int32 MaxGroupSize = CalcMaxNumGroupItems();
+
+	// Make sure we have the right number of groups.
+	const int32 NumCurveGroups = CurveGroups.Num();
+	if (OutCurveGroupIndices.Num() != NumCurveGroups * MaxGroupSize)
+	{
+		OutCurveGroupIndices.Reset();
+		OutCurveGroupIndices.AddDefaulted(NumCurveGroups * MaxGroupSize);
+	}
+
+	// Build a list of ints for every group.
+	// The integers in the group are the curve indices, for curve values we pass to the network as inputs.
+	// 0 means the first curve we pass to the network inputs.
+	// 9 means the tenth curve we pass to the network inputs, etc.
+	// All groups will have the same number of items. We pad them with INDEX_NONE when they have less items
+	// compared to the maximum number of items of all groups.
+	for (int32 CurveGroupIndex = 0; CurveGroupIndex < NumCurveGroups; ++CurveGroupIndex)
+	{
+		int32 LastValidIndex = INDEX_NONE;
+		const int32 NumItemsInGroup = CurveGroups[CurveGroupIndex].CurveNames.Num();
+		for (int32 Index = 0; Index < MaxGroupSize; Index++)
+		{
+			if (Index < NumItemsInGroup)
+			{
+				const FName CurveName = CurveGroups[CurveGroupIndex].CurveNames[Index].CurveName;
+				int32 CurveInputIndex = CurveName.IsValid() ? CurveNames.Find(CurveName) : INDEX_NONE;
+				if (CurveInputIndex != INDEX_NONE)
+				{
+					LastValidIndex = CurveInputIndex;
+				}
+				OutCurveGroupIndices[CurveGroupIndex * MaxGroupSize + Index] = CurveInputIndex;
+			}
+			else
+			{
+				OutCurveGroupIndices[CurveGroupIndex * MaxGroupSize + Index] = LastValidIndex;
+			}
+		}
+	}
+}
+
+bool UNeuralMorphInputInfo::HasInvalidGroups() const
+{
+	// Check all bone groups.
+	for (const FNeuralMorphBoneGroup& BoneGroup : BoneGroups)
+	{
+		if (BoneGroup.BoneNames.Num() <= 1)
+		{
+			return true;
+		}
+
+		// We require all bone names to be valid.
+		const int32 NumCurves = BoneGroup.BoneNames.Num();
+		for (int32 Index = 0; Index < NumCurves; ++Index)
+		{
+			const FName Name = BoneGroup.BoneNames[Index].BoneName;
+			if (!Name.IsValid() || Name.IsNone())
+			{
+				return true;
+			}
+
+			if (!BoneNames.Contains(Name))
+			{
+				return true;
+			}
+		}
+	}
+
+	// Check all curve groups.
+	for (const FNeuralMorphCurveGroup& CurveGroup : CurveGroups)
+	{
+		if (CurveGroup.CurveNames.Num() <= 1)
+		{
+			return true;
+		}
+
+		// We require all curve names to be valid.
+		const int32 NumCurves = CurveGroup.CurveNames.Num();
+		for (int32 Index = 0; Index < NumCurves; ++Index)
+		{
+			const FName Name = CurveGroup.CurveNames[Index].CurveName;
+			if (!Name.IsValid() || Name.IsNone())
+			{
+				return true;
+			}
+
+			if (!CurveNames.Contains(Name))
+			{
+				return true;
+			}
+		}
+	}
+
+	return false;
+}
diff --git a/Engine/Plugins/Animation/MLDeformer/NeuralMorphModel/Source/NeuralMorphModel/Private/NeuralMorphModel.cpp b/Engine/Plugins/Animation/MLDeformer/NeuralMorphModel/Source/NeuralMorphModel/Private/NeuralMorphModel.cpp
--- a/Engine/Plugins/Animation/MLDeformer/NeuralMorphModel/Source/NeuralMorphModel/Private/NeuralMorphModel.cpp
+++ b/Engine/Plugins/Animation/MLDeformer/NeuralMorphModel/Source/NeuralMorphModel/Private/NeuralMorphModel.cpp
@@ -2,6 +2,7 @@
 #include "NeuralMorphModel.h"
 #include "NeuralMorphModelVizSettings.h"
 #include "NeuralMorphModelInstance.h"
+#include "NeuralMorphInputInfo.h"
 #include "MLDeformerAsset.h"
 #include "MLDeformerComponent.h"
 #include "UObject/UObjectGlobals.h"
@@ -47,6 +48,11 @@ UMLDeformerModelInstance* UNeuralMorphModel::CreateModelInstance(UMLDeformerComp
 	return NewObject<UNeuralMorphModelInstance>(Component);
 }
 
+UMLDeformerInputInfo* UNeuralMorphModel::CreateInputInfo()
+{
+	return NewObject<UNeuralMorphInputInfo>(this);
+}
+
 void UNeuralMorphModel::Serialize(FArchive& Archive)
 {
 #if !NEURALMORPHMODEL_FORCE_USE_NNI
@@ -68,6 +74,18 @@ void UNeuralMorphModel::Serialize(FArchive& Archive)
 	}
 #endif
 
+	// Convert the UMLDeformerInputInfo object into a UNeuralMorphInputInfo object.
+	UMLDeformerInputInfo* CurInputInfo = GetInputInfo();
+	if (CurInputInfo)
+	{
+		if (!CurInputInfo->IsA<UNeuralMorphInputInfo>())
+		{
+			UNeuralMorphInputInfo* NeuralMorphInputInfo = Cast<UNeuralMorphInputInfo>(CreateInputInfo());
+			NeuralMorphInputInfo->CopyMembersFrom(CurInputInfo);
+			SetInputInfo(NeuralMorphInputInfo);
+		}		
+	}
+
 	Super::Serialize(Archive);
 }
 
diff --git a/Engine/Plugins/Animation/MLDeformer/NeuralMorphModel/Source/NeuralMorphModel/Private/NeuralMorphModelInstance.cpp b/Engine/Plugins/Animation/MLDeformer/NeuralMorphModel/Source/NeuralMorphModel/Private/NeuralMorphModelInstance.cpp
--- a/Engine/Plugins/Animation/MLDeformer/NeuralMorphModel/Source/NeuralMorphModel/Private/NeuralMorphModelInstance.cpp
+++ b/Engine/Plugins/Animation/MLDeformer/NeuralMorphModel/Source/NeuralMorphModel/Private/NeuralMorphModelInstance.cpp
@@ -3,24 +3,124 @@
 #include "NeuralMorphModelInstance.h"
 #include "NeuralMorphModel.h"
 #include "NeuralMorphNetwork.h"
+#include "NeuralMorphInputInfo.h"
 #include "MLDeformerComponent.h"
 #include "MLDeformerAsset.h"
 #include "Animation/AnimInstance.h"
 #include "Components/ExternalMorphSet.h"
 #include "Components/SkeletalMeshComponent.h"
 
-
 void UNeuralMorphModelInstance::Init(USkeletalMeshComponent* SkelMeshComponent)
 {
 	Super::Init(SkelMeshComponent);
 
-#if !NEURALMORPHMODEL_FORCE_USE_NNI
 	UNeuralMorphModel* MorphModel = Cast<UNeuralMorphModel>(Model);
+
+#if !NEURALMORPHMODEL_FORCE_USE_NNI
 	UNeuralMorphNetwork* MorphNetwork = MorphModel->GetNeuralMorphNetwork();
 	NetworkInstance = MorphNetwork ? MorphNetwork->CreateInstance() : nullptr;
 #endif
+
+	UNeuralMorphInputInfo* InputInfo = Cast<UNeuralMorphInputInfo>(Model->GetInputInfo());
+	if (InputInfo == nullptr)
+	{
+		return;
+	}
+
+	InputInfo->GenerateBoneGroupIndices(BoneGroupIndices);
+	InputInfo->GenerateCurveGroupIndices(CurveGroupIndices);
 }
 
+void UNeuralMorphModelInstance::FillNetworkInputs()
+{
+	check(NetworkInstance);
+	UNeuralMorphNetwork* MorphNetwork = Cast<UNeuralMorphNetwork>(NetworkInstance->GetOuter());
+
+	// Write the input bone transforms.
+	float* MainInputs = NetworkInstance->GetInputs().GetData();
+	const UNeuralMorphMLP* MainMLP = MorphNetwork->GetMainMLP();
+	const int32 NumMainMLPInputs = MainMLP->GetNumInputs();
+	Super::SetBoneTransforms(MainInputs, NumMainMLPInputs, 0);
+
+	// Write the curve input values, start writing after the bone inputs.
+	const int32 CurveWriteOffset = Model->GetInputInfo()->GetNumBones() * 6;	// 6 Floats per bone.
+	SetCurveValues(MainInputs, NumMainMLPInputs, CurveWriteOffset);
+
+	// Normalize the inputs.
+	const float* Means = MorphNetwork->GetInputMeans().GetData();
+	const float* Stds = MorphNetwork->GetInputStds().GetData();
+	const int32 NumInputs = MorphNetwork->GetInputMeans().Num();
+	for (int32 Index = 0; Index < NumInputs; ++Index)
+	{
+		MainInputs[Index] = (MainInputs[Index] - Means[Index]) / Stds[Index];
+	}
+
+	// Write the input transforms for the group MLP.
+	// We reuse the same input values from the main network, and those values just got normalized.
+	// So there is no need to normalize them again.
+	const UNeuralMorphMLP* GroupMLP = MorphNetwork->GetGroupMLP();
+	if (GroupMLP)
+	{
+		// Only the local mode should be using the group MLP.
+		check(Cast<UNeuralMorphModel>(Model)->GetModelMode() == ENeuralMorphMode::Local);
+
+		// Write the bone transforms.
+		int32 Offset = 0;
+		float* GroupNetworkInputs = NetworkInstance->GetGroupInputs().GetData();
+		const int32 NumBoneGroups = BoneGroupIndices.Num();
+		for (int32 Index = 0; Index < BoneGroupIndices.Num(); ++Index)
+		{
+			const int32 BoneIndex = BoneGroupIndices[Index];
+			if (BoneIndex != INDEX_NONE)
+			{
+				const int32 BoneOffset = BoneIndex * 6;
+				GroupNetworkInputs[Offset++] = MainInputs[BoneOffset];
+				GroupNetworkInputs[Offset++] = MainInputs[BoneOffset + 1];
+				GroupNetworkInputs[Offset++] = MainInputs[BoneOffset + 2];
+				GroupNetworkInputs[Offset++] = MainInputs[BoneOffset + 3];
+				GroupNetworkInputs[Offset++] = MainInputs[BoneOffset + 4];
+				GroupNetworkInputs[Offset++] = MainInputs[BoneOffset + 5];
+			}
+			else
+			{
+				GroupNetworkInputs[Offset++] = 0.0f;
+				GroupNetworkInputs[Offset++] = 0.0f;
+				GroupNetworkInputs[Offset++] = 0.0f;
+				GroupNetworkInputs[Offset++] = 0.0f;
+				GroupNetworkInputs[Offset++] = 0.0f;
+				GroupNetworkInputs[Offset++] = 0.0f;
+			}
+		}
+
+		// Write the curve values.
+		const int32 NumFloatsPerCurve = MorphNetwork->GetNumFloatsPerCurve();
+		check(MorphNetwork->GetNumFloatsPerCurve() == 6);	// Local mode curve values are expected to use 6 floats.
+		const int32 NumCurveGroups = CurveGroupIndices.Num();
+		for (int32 Index = 0; Index < CurveGroupIndices.Num(); ++Index)
+		{
+			const int32 CurveIndex = CurveGroupIndices[Index];
+			if (CurveIndex != INDEX_NONE)
+			{
+				const int32 CurveOffset = CurveWriteOffset + (CurveIndex * NumFloatsPerCurve);
+				GroupNetworkInputs[Offset++] = MainInputs[CurveOffset];
+				GroupNetworkInputs[Offset++] = MainInputs[CurveOffset + 1];
+				GroupNetworkInputs[Offset++] = MainInputs[CurveOffset + 2];
+				GroupNetworkInputs[Offset++] = MainInputs[CurveOffset + 3];
+				GroupNetworkInputs[Offset++] = MainInputs[CurveOffset + 4];
+				GroupNetworkInputs[Offset++] = MainInputs[CurveOffset + 5];
+			}
+			else
+			{
+				GroupNetworkInputs[Offset++] = 0.0f;
+				GroupNetworkInputs[Offset++] = 0.0f;
+				GroupNetworkInputs[Offset++] = 0.0f;
+				GroupNetworkInputs[Offset++] = 0.0f;
+				GroupNetworkInputs[Offset++] = 0.0f;
+				GroupNetworkInputs[Offset++] = 0.0f;
+			}
+		}
+	}
+}
 
 int64 UNeuralMorphModelInstance::SetCurveValues(float* OutputBuffer, int64 OutputBufferSize, int64 StartIndex)
 {
@@ -43,7 +143,7 @@ int64 UNeuralMorphModelInstance::SetCurveValues(float* OutputBuffer, int64 Outpu
 			OutputBuffer[Index++] = 0.0f;
 		}
 
-		// Write the 
+		// Write the actual curve values.
 		UAnimInstance* AnimInstance = SkeletalMeshComponent->GetAnimInstance();
 		if (AnimInstance)
 		{
@@ -100,6 +200,7 @@ bool UNeuralMorphModelInstance::SetupInputs()
 	// Some safety checks.
 	if (SkeletalMeshComponent == nullptr ||
 		SkeletalMeshComponent->GetSkeletalMeshAsset() == nullptr ||
+		MorphNetwork->GetMainMLP() == nullptr ||
 		!bIsCompatible)
 	{
 		return false;
@@ -115,10 +216,8 @@ bool UNeuralMorphModelInstance::SetupInputs()
 		return false;
 	}
 
-	// Update and write the input values directly into the input tensor.
-	float* InputDataPointer = NetworkInstance->GetInputs().GetData();
-	const int64 NumFloatsWritten = SetNeuralNetworkInputValues(InputDataPointer, NumNeuralNetInputs);
-	check(NumFloatsWritten == NumNeuralNetInputs);
+	// Set the actual network input values.
+	FillNetworkInputs();
 
 	return true;
 #endif
@@ -149,59 +248,60 @@ void UNeuralMorphModelInstance::Execute(float ModelWeight)
 		return;
 	}
 
-	if (!MorphNetwork->IsEmpty())
+	// Make sure we have a valid network instance.
+	if (NetworkInstance == nullptr)
+	{
+		WeightData->ZeroWeights();
+		return;
+	}
+
+	// If our model is active, we want to run the neural network and update the morph weights
+	// with the values that the neural net calculated for us.
+	// Get the network output values, read the values and use them as morph target weights inside the skeletal mesh component.
+	const TArrayView<const float> NetworkOutputs = NetworkInstance->GetOutputs();
+	const int32 NumNetworkWeights = NetworkOutputs.Num();
+	const int32 NumMorphTargets = WeightData->Weights.Num();
+
+	// Set the weights to zero if we have no valid network.
+	if (MorphNetwork->IsEmpty() || NumMorphTargets != NumNetworkWeights + 1)	// Plus 1 because of the extra means morph target.
+	{
+		WeightData->ZeroWeights();
+		return;
+	}
+
+	// Perform inference on the neural network, this updates its output values.
+	NetworkInstance->Run();
+
+	// Set the first morph target, which represents the means, to a weight of 1.0, as it always needs to be fully active.
+	WeightData->Weights[0] = ModelWeight;
+
+	// Update all generated morph target weights with the values calculated by our neural network.
+	const TArrayView<const float> ErrorValues = MorphModel->GetMorphTargetErrorValues();
+	const TArrayView<const int32> ErrorOrder = MorphModel->GetMorphTargetErrorOrder();
+	if (!ErrorValues.IsEmpty())
 	{
-		// Perform inference on the neural network, this updates its output values.
-		NetworkInstance->Run();
-
-		// If our model is active, we want to run the neural network and update the morph weights
-		// with the values that the neural net calculated for us.
-		// Get the network output values, read the values and use them as morph target weights inside the skeletal mesh component.
-		const TArrayView<const float> NetworkOutputs = NetworkInstance->GetOutputs();
-		const int32 NumNetworkWeights = NetworkOutputs.Num();
-		const int32 NumMorphTargets = WeightData->Weights.Num();
-
-		// If we have the expected amount of weights.
-		// +1 because we always have an extra morph target that represents the means, with fixed weight of 1.
-		// Therefore, the neural network output will contain one less float than the number of morph targets in our morph set.
-		if (NumMorphTargets == NumNetworkWeights + 1)
+		const int32 QualityLevel = GetMLDeformerComponent()->GetQualityLevel();
+		const int32 NumActiveMorphs = MorphModel->GetNumActiveMorphs(QualityLevel);
+		for (int32 Index = 0; Index < NumActiveMorphs; ++Index)
 		{
-			// Set the first morph target, which represents the means, to a weight of 1.0, as it always needs to be fully active.
-			WeightData->Weights[0] = ModelWeight;
+			const int32 MorphIndex = ErrorOrder[Index];
+			const float TargetWeight = NetworkOutputs[MorphIndex] * ModelWeight;
+			WeightData->Weights[MorphIndex + 1] = FMath::Lerp<float, float>(StartMorphWeights[MorphIndex + 1], TargetWeight, MorphLerpAlpha);
+		}
 
-			// Update all generated morph target weights with the values calculated by our neural network.
-			const TArrayView<const float> ErrorValues = MorphModel->GetMorphTargetErrorValues();
-			const TArrayView<const int32> ErrorOrder = MorphModel->GetMorphTargetErrorOrder();
-			if (!ErrorValues.IsEmpty())
-			{
-				const int32 QualityLevel = GetMLDeformerComponent()->GetQualityLevel();
-				const int32 NumActiveMorphs = MorphModel->GetNumActiveMorphs(QualityLevel);
-				for (int32 Index = 0; Index < NumActiveMorphs; ++Index)
-				{
-					const int32 MorphIndex = ErrorOrder[Index];
-					const float TargetWeight = NetworkOutputs[MorphIndex] * ModelWeight;
-					WeightData->Weights[MorphIndex + 1] = FMath::Lerp<float, float>(StartMorphWeights[MorphIndex + 1], TargetWeight, MorphLerpAlpha);
-				}
-
-				// Disable all inactive morphs.
-				for (int32 Index = NumActiveMorphs; Index < NumNetworkWeights; ++Index)
-				{
-					const int32 MorphIndex = ErrorOrder[Index];
-					WeightData->Weights[MorphIndex + 1] = FMath::Lerp<float, float>(StartMorphWeights[MorphIndex + 1], 0.0f, MorphLerpAlpha);
-				}
-			}
-			else
-			{
-				for (int32 MorphIndex = 0; MorphIndex < NumNetworkWeights; ++MorphIndex)
-				{
-					WeightData->Weights[MorphIndex + 1] = NetworkOutputs[MorphIndex] * ModelWeight;
-				}
-			}
-			return;
+		// Disable all inactive morphs.
+		for (int32 Index = NumActiveMorphs; Index < NumNetworkWeights; ++Index)
+		{
+			const int32 MorphIndex = ErrorOrder[Index];
+			WeightData->Weights[MorphIndex + 1] = FMath::Lerp<float, float>(StartMorphWeights[MorphIndex + 1], 0.0f, MorphLerpAlpha);
+		}
+	}
+	else // Old models might not have the error values yet, and thus not support the quality levels yet.
+	{
+		for (int32 MorphIndex = 0; MorphIndex < NumNetworkWeights; ++MorphIndex)
+		{
+			WeightData->Weights[MorphIndex + 1] = NetworkOutputs[MorphIndex] * ModelWeight;
 		}
 	}
-
-	// Set the weights to zero if we have no valid network.
-	WeightData->ZeroWeights();
 #endif
 }
diff --git a/Engine/Plugins/Animation/MLDeformer/NeuralMorphModel/Source/NeuralMorphModel/Private/NeuralMorphNetwork.cpp b/Engine/Plugins/Animation/MLDeformer/NeuralMorphModel/Source/NeuralMorphModel/Private/NeuralMorphNetwork.cpp
--- a/Engine/Plugins/Animation/MLDeformer/NeuralMorphModel/Source/NeuralMorphModel/Private/NeuralMorphNetwork.cpp
+++ b/Engine/Plugins/Animation/MLDeformer/NeuralMorphModel/Source/NeuralMorphModel/Private/NeuralMorphNetwork.cpp
@@ -10,30 +10,214 @@
 	#include "NeuralMorphNetwork.ispc.generated.h"
 #endif
 
+//--------------------------------------------------------------------------
+// UNeuralMorphMLP
+//--------------------------------------------------------------------------
+void UNeuralMorphMLP::Empty()
+{
+	Layers.Empty();
+}
+
+bool UNeuralMorphMLP::IsEmpty() const
+{
+	return GetNumLayers() == 0;
+}
+
+int32 UNeuralMorphMLP::GetNumInputs() const
+{
+	return !Layers.IsEmpty() ? Layers[0]->NumInputs * Layers[0]->Depth : 0;
+}
+
+int32 UNeuralMorphMLP::GetNumOutputs() const
+{
+	return !Layers.IsEmpty() ? Layers[Layers.Num() - 1]->NumOutputs * Layers[Layers.Num() - 1]->Depth : 0;
+}
+
+int32 UNeuralMorphMLP::GetMaxNumLayerInputs() const
+{
+	int32 MaxValue = 0;
+	for (UNeuralMorphMLPLayer* Layer : Layers)
+	{
+		MaxValue = FMath::Max<int32>(MaxValue, Layer->NumInputs * Layer->Depth);
+	}
+	return MaxValue;
+}
+
+int32 UNeuralMorphMLP::GetMaxNumLayerOutputs() const
+{
+	int32 MaxValue = 0;
+	for (UNeuralMorphMLPLayer* Layer : Layers)
+	{
+		MaxValue = FMath::Max<int32>(MaxValue, Layer->NumOutputs * Layer->Depth);
+	}
+	return MaxValue;
+}
+
+UNeuralMorphMLPLayer* UNeuralMorphMLP::LoadNetworkLayer(FArchive& FileReader)
+{
+	UNeuralMorphMLPLayer* Layer = NewObject<UNeuralMorphMLPLayer>(this);
+
+	struct FLayerHeader
+	{
+		int32 NumInputs;
+		int32 NumOutputs;
+		int32 NumWeights;
+		int32 NumBiases;
+	};
+
+	FLayerHeader LayerHeader;
+	FileReader.Serialize(&LayerHeader, sizeof(FLayerHeader));
+	if (FileReader.IsError())
+	{
+		UE_LOG(LogNeuralMorphModel, Error, TEXT("Failed to load a layer header for the network!"));
+		Layer->ConditionalBeginDestroy();
+		return nullptr;
+	}
+
+	Layer->NumInputs = LayerHeader.NumInputs;
+	Layer->NumOutputs = LayerHeader.NumOutputs;
+	check(LayerHeader.NumInputs * LayerHeader.NumOutputs > 0);
+	Layer->Depth = LayerHeader.NumWeights / (LayerHeader.NumInputs * LayerHeader.NumOutputs);
+
+	Layer->Weights.AddZeroed(LayerHeader.NumWeights);
+	FileReader.Serialize(Layer->Weights.GetData(), LayerHeader.NumWeights * sizeof(float));
+	if (FileReader.IsError())
+	{
+		UE_LOG(LogNeuralMorphModel, Error, TEXT("Failed to load layer weights of the network!"));
+		Layer->ConditionalBeginDestroy();
+		return nullptr;
+	}
+
+	Layer->Biases.AddZeroed(LayerHeader.NumBiases);
+	FileReader.Serialize(Layer->Biases.GetData(), LayerHeader.NumBiases * sizeof(float));
+	if (FileReader.IsError())
+	{
+		UE_LOG(LogNeuralMorphModel, Error, TEXT("Failed to load layer biases of the network!"));
+		Layer->ConditionalBeginDestroy();
+		return nullptr;
+	}
+
+	return Layer;
+}
+
+
+bool UNeuralMorphMLP::Load(FArchive& FileReader)
+{
+	Empty();
+
+	// Read the FOURCC, to identify the file type.
+	char FOURCC[4] {' ', ' ', ' ', ' '};
+	FileReader.Serialize(FOURCC, 4);
+	if (FileReader.IsError())
+	{
+		UE_LOG(LogNeuralMorphModel, Error, TEXT("Failed to read the FOURCC!"));
+		Empty();
+		return false;
+	}
+
+	if (FOURCC[0] != 'N' || FOURCC[1] != 'M' || FOURCC[2] != 'L' || FOURCC[3] != 'P')	// NMLP (Neural MLP)
+	{
+		UE_LOG(LogNeuralMorphModel, Error, TEXT("The file is not a valid valid neural morph MLP file type!"));
+		Empty();
+		return false;
+	}
+
+	// Load the version number.
+	int32 Version = -1;
+	FileReader.Serialize(&Version, sizeof(int32));
+	if (FileReader.IsError())
+	{
+		UE_LOG(LogNeuralMorphModel, Error, TEXT("Failed to load version"));
+		Empty();
+		return false;
+	}
+
+	if (Version != 1)
+	{
+		UE_LOG(LogNeuralMorphModel, Error, TEXT("The Neural Morph MLP file '%s' is of an unknown version (Version=%d)!"), *FileReader.GetArchiveName(), Version);
+		Empty();
+		return false;
+	}
+
+	int32 NumLayers = -1;
+	FileReader.Serialize(&NumLayers, sizeof(int32));
+	if (FileReader.IsError())
+	{
+		UE_LOG(LogNeuralMorphModel, Error, TEXT("Failed to load number of layers"));
+		Empty();
+		return false;
+	}
+
+	// Load the weights and biases of the main network.
+	for (int32 LayerIndex = 0; LayerIndex < NumLayers; ++LayerIndex)
+	{
+		UNeuralMorphMLPLayer* Layer = LoadNetworkLayer(FileReader);
+		if (Layer == nullptr)
+		{
+			Empty();
+			return false;
+		}
+		Layers.Add(Layer);
+		UE_LOG(LogNeuralMorphModel, Verbose, TEXT("MLP Layer %d --> NumWeights=%d (%dx%dx%d)   NumBiases=%d"), LayerIndex, Layer->Weights.Num(), Layer->NumInputs, Layer->NumOutputs, Layer->Depth, Layer->Biases.Num());
+	}
+
+	UE_LOG(LogNeuralMorphModel, Display, TEXT("Successfullly loaded neural morph MLP from file '%s'"), *FileReader.GetArchiveName());
+	return true;
+}
+
+int32 UNeuralMorphMLP::GetNumLayers() const
+{
+	return Layers.Num();
+}
+
+UNeuralMorphMLPLayer& UNeuralMorphMLP::GetLayer(int32 Index) const
+{
+	return *Layers[Index].Get();
+}
+
+//--------------------------------------------------------------------------
+// UNeuralMorphNetwork
+//--------------------------------------------------------------------------
 void UNeuralMorphNetwork::Empty()
 {
 	InputMeans.Empty();
 	InputStd.Empty();
-	Layers.Empty();
+
+	if (MainMLP)
+	{
+		MainMLP->ConditionalBeginDestroy();
+		MainMLP = nullptr;
+	}
+
+	if (GroupMLP)
+	{
+		GroupMLP->ConditionalBeginDestroy();
+		GroupMLP = nullptr;
+	}
+
 	Mode = ENeuralMorphMode::Global;
 	NumMorphsPerBone = 0;
 	NumBones = 0;
 	NumCurves = 0;
+	NumGroups = 0;
+	NumItemsPerGroup = 0;
+	NumFloatsPerCurve = 1;
 }
 
 bool UNeuralMorphNetwork::IsEmpty() const
 {
-	return GetNumInputs() == 0;
+	return !MainMLP;
 }
 
-int32 UNeuralMorphNetwork::GetNumInputs() const
+UNeuralMorphMLP* UNeuralMorphNetwork::LoadMLP(FArchive& FileReader)
 {
-	return !Layers.IsEmpty() ? Layers[0]->NumInputs * Layers[0]->Depth : 0;
-}
-
-int32 UNeuralMorphNetwork::GetNumOutputs() const
-{
-	return !Layers.IsEmpty() ? Layers[Layers.Num() - 1]->NumOutputs * Layers[Layers.Num() - 1]->Depth : 0;
+	UNeuralMorphMLP* MLP = NewObject<UNeuralMorphMLP>(this);
+	if (!MLP->Load(FileReader))
+	{
+		MLP->ConditionalBeginDestroy();
+		MLP = nullptr;		
+	}
+	return MLP;
 }
 
 bool UNeuralMorphNetwork::Load(const FString& Filename)
@@ -80,7 +264,7 @@ bool UNeuralMorphNetwork::Load(const FString& Filename)
 		return false;
 	}
 
-	if (Version != 2)
+	if (Version != 3)
 	{
 		UE_LOG(LogNeuralMorphModel, Error, TEXT("The Neural Morph Network file '%s' is of an unknown version (Version=%d)!"), *Filename, Version);
 		FileReader->Close();
@@ -91,15 +275,17 @@ bool UNeuralMorphNetwork::Load(const FString& Filename)
 	// Load the header with info.
 	struct FInfoHeader
 	{
-		int32 Mode = -1;					// 0=Local, 1=Global
-		int32 NumInputs = -1;				// How many float inputs in the input layer?
-		int32 NumHiddenLayers = -1;			// How many hidden layers in the network (layers excluding input and output layer)?
-		int32 NumUnitsPerHiddenLayer = -1;	// How many neurons for each hidden layer?
-		int32 NumOutputs = -1;				// The number of units in the output layer.
-		int32 NumMorphsPerBone = -1;		// The number of morph targets per bone, if set Mode == 0. Otherwise ignored.
-		int32 NumBones = -1;				// The number of bones used as input.
-		int32 NumCurves = -1;				// The number of curves used as input.
-		int32 NumFloatsPerCurve = -1;		// The number of floats per curve.
+		int32 Mode;						// 0=Local, 1=Global
+		int32 NumInputs;				// How many float inputs in the input layer?
+		int32 NumHiddenLayers;			// How many hidden layers in the network (layers excluding input and output layer)?
+		int32 NumUnitsPerHiddenLayer;	// How many neurons for each hidden layer?
+		int32 NumOutputs;				// The number of units in the output layer.
+		int32 NumMorphsPerBone;			// The number of morph targets per bone, if set Mode == 0. Otherwise ignored.
+		int32 NumBones;					// The number of bones used as input.
+		int32 NumCurves;				// The number of curves used as input.
+		int32 NumFloatsPerCurve;		// The number of floats per curve.
+		int32 NumGroups;				// How many groups are there?
+		int32 NumItemsPerGroup;			// How many items (bones/curves) per group?
 	};
 
 	FInfoHeader Info;
@@ -112,19 +298,13 @@ bool UNeuralMorphNetwork::Load(const FString& Filename)
 		return false;
 	}
 
-	if (Info.Mode == -1 || Info.NumInputs == -1 || Info.NumHiddenLayers == -1 || Info.NumUnitsPerHiddenLayer == -1 || Info.NumOutputs == -1 || Info.NumBones == -1 || Info.NumCurves == -1 || Info.NumFloatsPerCurve == -1)
-	{
-		UE_LOG(LogNeuralMorphModel, Error, TEXT("Failed to read info header for Neural Morph Network in file '%s'!"), *FileReader->GetArchiveName());
-		FileReader->Close();
-		Empty();
-		return false;
-	}
-
 	Mode = (Info.Mode == 0) ? ENeuralMorphMode::Local : ENeuralMorphMode::Global;
 	NumMorphsPerBone = Info.NumMorphsPerBone;
 	NumBones = Info.NumBones;
 	NumCurves = Info.NumCurves;
 	NumFloatsPerCurve = Info.NumFloatsPerCurve;
+	NumGroups = Info.NumGroups;
+	NumItemsPerGroup = Info.NumItemsPerGroup;
 
 	// Read the input standard deviation and means.
 	InputMeans.SetNumZeroed(Info.NumInputs);
@@ -147,56 +327,27 @@ bool UNeuralMorphNetwork::Load(const FString& Filename)
 		return false;
 	}
 
-	struct FLayerHeader
+	// Load the main network.
+	MainMLP = LoadMLP(*FileReader);
+	if (!MainMLP)
 	{
-		int32 NumInputs = -1;
-		int32 NumOutputs = -1;
-		int32 NumWeights = -1;
-		int32 NumBiases = -1;
-	};
+		UE_LOG(LogNeuralMorphModel, Error, TEXT("Failed to load main MLP!"));
+		FileReader->Close();
+		Empty();
+		return false;
+	}
 
-	// Load the weights and biases.
-	for (int32 LayerIndex = 0; LayerIndex < Info.NumHiddenLayers + 1; ++LayerIndex)
+	// If we're in local mode and we have groups, load the network for that.
+	if (NumGroups > 0 && Mode == ENeuralMorphMode::Local)
 	{
-		UNeuralMorphNetworkLayer* WeightData = NewObject<UNeuralMorphNetworkLayer>(this);
-
-		FLayerHeader LayerHeader;
-		FileReader->Serialize(&LayerHeader, sizeof(FLayerHeader));
-		if (FileReader->IsError())
+		GroupMLP = LoadMLP(*FileReader);
+		if (!GroupMLP)
 		{
-			UE_LOG(LogNeuralMorphModel, Error, TEXT("Failed to load layer header!"));
+			UE_LOG(LogNeuralMorphModel, Error, TEXT("Failed to load group MLP!"));
 			FileReader->Close();
 			Empty();
 			return false;
 		}
-
-		WeightData->NumInputs = LayerHeader.NumInputs;
-		WeightData->NumOutputs = LayerHeader.NumOutputs;
-		WeightData->Depth = LayerHeader.NumWeights / (LayerHeader.NumInputs * LayerHeader.NumOutputs);
-
-		WeightData->Weights.AddZeroed(LayerHeader.NumWeights);
-		FileReader->Serialize(WeightData->Weights.GetData(), LayerHeader.NumWeights * sizeof(float));
-		if (FileReader->IsError())
-		{
-			UE_LOG(LogNeuralMorphModel, Error, TEXT("Failed to load layer weights!"));
-			FileReader->Close();
-			Empty();
-			return false;
-		}
-
-		WeightData->Biases.AddZeroed(LayerHeader.NumBiases);
-		FileReader->Serialize(WeightData->Biases.GetData(), LayerHeader.NumBiases * sizeof(float));
-		if (FileReader->IsError())
-		{
-			UE_LOG(LogNeuralMorphModel, Error, TEXT("Failed to load layer biases!"));
-			FileReader->Close();
-			Empty();
-			return false;
-		}
-
-		Layers.Add(WeightData);
-
-		UE_LOG(LogNeuralMorphModel, Verbose, TEXT("Network Layer %d --> NumWeights=%d (%dx%dx%d)   NumBiases=%d"), LayerIndex, WeightData->Weights.Num(), WeightData->NumInputs, WeightData->NumOutputs, WeightData->Depth, WeightData->Biases.Num());
 	}
 
 	// Clean up and return the result.
@@ -251,21 +402,44 @@ const TArrayView<const float> UNeuralMorphNetwork::GetInputStds() const
 	return InputStd;
 }
 
-int32 UNeuralMorphNetwork::GetNumLayers() const
+int32 UNeuralMorphNetwork::GetNumFloatsPerCurve() const
 {
-	return Layers.Num();
+	return NumFloatsPerCurve;
 }
 
-UNeuralMorphNetworkLayer& UNeuralMorphNetwork::GetLayer(int32 Index) const
+int32 UNeuralMorphNetwork::GetNumGroups() const
 {
-	return *Layers[Index].Get();
+	return NumGroups;
 }
 
-int32 UNeuralMorphNetwork::GetNumFloatsPerCurve() const
+int32 UNeuralMorphNetwork::GetNumItemsPerGroup() const
 {
-	return NumFloatsPerCurve;
+	return NumItemsPerGroup;
+}
+
+int32 UNeuralMorphNetwork::GetNumInputs() const
+{
+	return MainMLP ? MainMLP->GetNumInputs() : 0;
 }
 
+int32 UNeuralMorphNetwork::GetNumOutputs() const
+{
+	int32 NumOutputs = MainMLP ? MainMLP->GetNumOutputs() : 0;
+	NumOutputs += GroupMLP ? GroupMLP->GetNumOutputs() : 0; 
+	return NumOutputs;
+}
+
+UNeuralMorphMLP* UNeuralMorphNetwork::GetMainMLP() const
+{
+	return MainMLP.Get();
+}
+
+UNeuralMorphMLP* UNeuralMorphNetwork::GetGroupMLP() const
+{
+	return GroupMLP.Get();
+}
+
+
 //--------------------------------------------------------------------------
 // UNeuralMorphNetworkInstance
 //--------------------------------------------------------------------------
@@ -280,6 +454,16 @@ TArrayView<const float> UNeuralMorphNetworkInstance::GetInputs() const
 	return TArrayView<const float>(Inputs.GetData(), Inputs.Num());
 }
 
+TArrayView<float> UNeuralMorphNetworkInstance::GetGroupInputs()
+{ 
+	return TArrayView<float>(GroupInputs.GetData(), GroupInputs.Num());
+}
+
+TArrayView<const float> UNeuralMorphNetworkInstance::GetGroupInputs() const
+{ 
+	return TArrayView<const float>(GroupInputs.GetData(), GroupInputs.Num());
+}
+
 TArrayView<float> UNeuralMorphNetworkInstance::GetOutputs()
 { 
 	return TArrayView<float>(Outputs.GetData(), Outputs.Num());
@@ -298,20 +482,27 @@ const UNeuralMorphNetwork& UNeuralMorphNetworkInstance::GetNeuralNetwork() const
 void UNeuralMorphNetworkInstance::Init(UNeuralMorphNetwork* InNeuralNetwork)
 {
 	Network = InNeuralNetwork;
+	const UNeuralMorphMLP* MainMLP = Network->GetMainMLP();
+	const UNeuralMorphMLP* GroupMLP = Network->GetGroupMLP();
+
 	Inputs.SetNumZeroed(Network->GetNumInputs());
-	Outputs.SetNumZeroed(Network->GetNumOutputs());
+	Outputs.SetNumZeroed(Network->GetNumOutputs());	// This contains concatenated outputs of both main and group MLP's.
 
-	// Find the largest layer unit size and pre-allocate a buffer of that size.
-	int32 MaxNumUnits = 0;
-	const int32 NumLayers = Network->GetNumLayers();
-	for (int32 LayerIndex = 0; LayerIndex < NumLayers; ++LayerIndex)
+	if (GroupMLP)
 	{
-		const UNeuralMorphNetworkLayer& CurLayer = Network->GetLayer(LayerIndex);
-		const int32 NumInputUnits = CurLayer.NumInputs * CurLayer.Depth;
-		const int32 NumOutputUnits = CurLayer.NumOutputs * CurLayer.Depth;
-		MaxNumUnits = FMath::Max3<int32>(NumInputUnits, NumOutputUnits, MaxNumUnits);
+		GroupInputs.SetNumZeroed(GroupMLP->GetNumInputs());
 	}
 
+	// Find the largest layer unit size and pre-allocate a buffer of that size.
+	const int32 MainMaxNumInputs = MainMLP ? MainMLP->GetMaxNumLayerInputs() : 0;
+	const int32 MainMaxNumOutputs = MainMLP ? MainMLP->GetMaxNumLayerOutputs() : 0;
+	const int32 GroupMaxNumInputs = GroupMLP ? GroupMLP->GetMaxNumLayerInputs() : 0;
+	const int32 GroupMaxNumOutputs = GroupMLP ? GroupMLP->GetMaxNumLayerOutputs() : 0;
+
+	// Find the maximum of the inputs and outputs of both main and group mlp's.
+	int32 MaxNumUnits = FMath::Max<int32>(MainMaxNumInputs, GroupMaxNumInputs);
+	MaxNumUnits = FMath::Max3<int32>(MaxNumUnits, MainMaxNumOutputs, GroupMaxNumOutputs);
+
 	TempInputArray.SetNumZeroed(MaxNumUnits);
 	TempOutputArray.SetNumZeroed(MaxNumUnits);
 }
@@ -323,22 +514,22 @@ void UNeuralMorphNetworkInstance::RunGlobalModel(const FRunSettings& RunSettings
 
 	checkfSlow(Network->GetNumFloatsPerCurve() == 1, TEXT("Expecting the number of floats per curve to be 1 in global mode."));
 
-	const int32 NumLayers = Network->GetNumLayers();
+	const UNeuralMorphMLP& MainMLP = *Network->GetMainMLP();
+
+	const int32 NumLayers = MainMLP.GetNumLayers();
 	for (int32 LayerIndex = 0; LayerIndex < NumLayers; ++LayerIndex)
 	{
-		const UNeuralMorphNetworkLayer& CurLayer = Network->GetLayer(LayerIndex);
+		const UNeuralMorphMLPLayer& CurLayer = MainMLP.GetLayer(LayerIndex);
 		const int32 NumLayerInputs = CurLayer.NumInputs;
 		const int32 NumLayerOutputs = CurLayer.NumOutputs;
 
-		// Normalize inputs for the first layer inputs.
+		// Copy the inputs to the temp input buffer for the first layer.
 		if (LayerIndex == 0)
 		{
 			const float* const RESTRICT NetworkInputs = RunSettings.InputBuffer;
-			const float* const RESTRICT Means = RunSettings.InputMeansBuffer;
-			const float* const RESTRICT Stds = RunSettings.InputStdsBuffer;
 			for (int32 Index = 0; Index < NumLayerInputs; ++Index)
 			{
-				TempInputBuffer[Index] = (NetworkInputs[Index] - Means[Index]) / Stds[Index];
+				TempInputBuffer[Index] = NetworkInputs[Index];
 			}
 		}
 
@@ -391,32 +582,29 @@ void UNeuralMorphNetworkInstance::RunGlobalModel(const FRunSettings& RunSettings
 	}
 }
 
-
-void UNeuralMorphNetworkInstance::RunLocalModel(const FRunSettings& RunSettings)
+void UNeuralMorphNetworkInstance::RunLocalMLP(UNeuralMorphMLP& MLP, const FRunSettings& RunSettings)
 {
 	float* RESTRICT TempInputBuffer = RunSettings.TempInputBuffer;
 	float* RESTRICT TempOutputBuffer = RunSettings.TempOutputBuffer;
 
 	checkfSlow(Network->GetNumFloatsPerCurve() == 6, TEXT("Expecting num floats per curve to be 6 in local mode."));
 
-	const int32 NumLayers = Network->GetNumLayers();
+	const int32 NumLayers = MLP.GetNumLayers();
 	for (int32 LayerIndex = 0; LayerIndex < NumLayers; ++LayerIndex)
 	{
-		const UNeuralMorphNetworkLayer& CurLayer = Network->GetLayer(LayerIndex);
+		const UNeuralMorphMLPLayer& CurLayer = MLP.GetLayer(LayerIndex);
 		const int32 NumInputsPerBlock = CurLayer.NumInputs;
 		checkSlow(NumInputsPerBlock == 6);
 		const int32 NumBlocks = CurLayer.Depth;
 
-		// Normalize inputs for the first layer inputs.
+		// Copy the inputs to the temp buffer.
 		if (LayerIndex == 0)
 		{
 			const float* const RESTRICT NetworkInputs = RunSettings.InputBuffer;
-			const float* const RESTRICT Means = RunSettings.InputMeansBuffer;
-			const float* const RESTRICT Stds = RunSettings.InputStdsBuffer;
 			const int32 NumInputs = CurLayer.NumInputs * CurLayer.Depth;
 			for (int32 Index = 0; Index < NumInputs; ++Index)
 			{
-				TempInputBuffer[Index] = (NetworkInputs[Index] - Means[Index]) / Stds[Index];
+				TempInputBuffer[Index] = NetworkInputs[Index];
 			}
 		}
 
@@ -476,7 +664,6 @@ void UNeuralMorphNetworkInstance::RunLocalModel(const FRunSettings& RunSettings)
 	} // For all layers.
 }
 
-
 void UNeuralMorphNetworkInstance::Run()
 {	
 	TRACE_CPUPROFILER_EVENT_SCOPE(UNeuralMorphNetwork::Run)
@@ -486,8 +673,6 @@ void UNeuralMorphNetworkInstance::Run()
 	RunSettings.TempInputBuffer  = TempInputArray.GetData();
 	RunSettings.TempOutputBuffer = TempOutputArray.GetData();
 	RunSettings.InputBuffer		 = Inputs.GetData();
-	RunSettings.InputMeansBuffer = Network->GetInputMeans().GetData();
-	RunSettings.InputStdsBuffer  = Network->GetInputStds().GetData();
 	RunSettings.OutputBuffer	 = Outputs.GetData();
 
 	// Run the network.
@@ -498,6 +683,15 @@ void UNeuralMorphNetworkInstance::Run()
 	else
 	{
 		checkSlow(Network->GetMode() == ENeuralMorphMode::Local)
-		RunLocalModel(RunSettings);
+		RunLocalMLP(*Network->GetMainMLP(), RunSettings);
+		if (Network->GetGroupMLP())
+		{
+			FRunSettings GroupRunSettings;
+			GroupRunSettings.TempInputBuffer	= TempInputArray.GetData();
+			GroupRunSettings.TempOutputBuffer	= TempOutputArray.GetData();
+			GroupRunSettings.InputBuffer		= GroupInputs.GetData();
+			GroupRunSettings.OutputBuffer		= Outputs.GetData() + Network->GetMainMLP()->GetNumOutputs();	// Output values after the main network's outputs.
+			RunLocalMLP(*Network->GetGroupMLP(), GroupRunSettings);
+		}
 	}
 }
diff --git a/Engine/Plugins/Animation/MLDeformer/NeuralMorphModel/Source/NeuralMorphModel/Private/NeuralMorphNetwork.ispc b/Engine/Plugins/Animation/MLDeformer/NeuralMorphModel/Source/NeuralMorphModel/Private/NeuralMorphNetwork.ispc
--- a/Engine/Plugins/Animation/MLDeformer/NeuralMorphModel/Source/NeuralMorphModel/Private/NeuralMorphNetwork.ispc
+++ b/Engine/Plugins/Animation/MLDeformer/NeuralMorphModel/Source/NeuralMorphModel/Private/NeuralMorphNetwork.ispc
@@ -1,10 +1,13 @@
 // Copyright Epic Games, Inc. All Rights Reserved.
 
+// An approximation to Exp(-X).
 float InvExpApprox(const float X)
 {
 	return 1 / (1 + 1.00746054f * X + 0.45053901f * X * X + 0.25724632f * X * X * X);		
 }
 
+// Process an MLP layer by multiplying inputs by the weights and adding bias values and passing that
+// through an ELU activation function.
 export void MorphNeuralNetwork_LayerForward(
 	uniform float Outputs[],
 	const uniform float Inputs[],
@@ -34,6 +37,7 @@ export void MorphNeuralNetwork_LayerForward(
 	}
 }
 
+// Perform an ELU activation function on every element.
 export void MorphNeuralNetwork_Activation_ELU(uniform float InputsOutputs[], const uniform int Num)
 {
 	foreach (Index = 0 ... Num)
diff --git a/Engine/Plugins/Animation/MLDeformer/NeuralMorphModel/Source/NeuralMorphModel/Public/NeuralMorphInputInfo.h b/Engine/Plugins/Animation/MLDeformer/NeuralMorphModel/Source/NeuralMorphModel/Public/NeuralMorphInputInfo.h
new file mode 100644
--- /dev/null
+++ b/Engine/Plugins/Animation/MLDeformer/NeuralMorphModel/Source/NeuralMorphModel/Public/NeuralMorphInputInfo.h
@@ -0,0 +1,83 @@
+// Copyright Epic Games, Inc. All Rights Reserved.
+#pragma once
+
+#include "CoreTypes.h"
+#include "MLDeformerInputInfo.h"
+#include "NeuralMorphTypes.h"
+#include "NeuralMorphInputInfo.generated.h"
+
+/**
+ * The neural morph model's input info that contains all the data about the inputs to the network.
+ */
+UCLASS()
+class NEURALMORPHMODEL_API UNeuralMorphInputInfo
+	: public UMLDeformerInputInfo
+{
+	GENERATED_BODY()
+
+public:
+	// FMLDeformerInputInfo overrides.
+	virtual void Reset() override;
+	// ~END FMLDeformerInputInfooverrides.
+
+	/**
+	 * Get the bone groups.
+	 * Bone groups are groups of bones that generate morph targets all together as a group rather than one for each bone individually.
+	 * @return A reference to the bone groups array.
+	 */
+	TArray<FNeuralMorphBoneGroup>& GetBoneGroups();
+
+	/**
+	 * Get the curve groups.
+	 * Curve groups are groups of bones that generate morph targets all together as a group rather than one for each curve individually.
+	 * @return A reference to the curve groups array.
+	 */
+	TArray<FNeuralMorphCurveGroup>& GetCurveGroups();
+
+	/**
+	 * Calculate the maximum number of group items.
+	 * This will look in both the bone and curve groups.
+	 * So for example, if there is a curve group with 5 items, while other groups have only 2 or 3 items, this method will return 5.
+	 * @return The maximum item count of all groups.
+	 */
+	int32 CalcMaxNumGroupItems() const;
+
+	/**
+	 * Generate a list of index values that point inside the bone array.
+	 * So the values that are output are either bone or curve numbers.
+	 * The number indexes inside the array returned by GetBoneNames().
+	 * The indices are a continuous array, with each padded group concatenated.
+	 * If we have two groups, like: ['UpperArm_L', 'Clavicle_L', 'Spine3'], ['UpperArm_R', 'Clavicle_R'] then this method will output
+	 * something like: [10, 11, 8, 15, 16, 16]. Where each number is the index inside our bone list.
+	 * As you can see there is one value (number 16) that is repeated. This is because for every group we will output the same number of items.
+	 * We find the maximum group item count, using CalcMaxGroupItems() and output number of values for each group.
+	 * We simply repeat the last number one or more times if we have to pad the item count of that group. This is why the number 16 is repeated two times.
+	 * @param OutBoneGroupIndices The array we will fill with the right bone indices. It will automatically resize internally.
+	 */
+	void GenerateBoneGroupIndices(TArray<int32>& OutBoneGroupIndices);
+
+	/**
+	 * Generates the list of index values that point inside the curve array.
+	 * For more information please look at the GenerateBoneGroupIndices documentation.
+	 * @param OutCurveGroupIndices The array we will fill with the right curve indices. It will automatically resize internally.
+	 * @see GenerateBoneGroupIndices.
+	 */
+	void GenerateCurveGroupIndices(TArray<int32>& OutCurveGroupIndices);
+
+	/**
+	 * Check whether we have invalid groups or not.
+	 * A group is considered invalid when it is empty, has items with the name 'None', or invalid names, or when the 
+	 * items (bones or curve names) do not exist in the bone or curve list.
+	 * @return Returns true when there are invalid bone or curve groups, otherwise false is returned.
+	 */
+	bool HasInvalidGroups() const;
+
+protected:
+	/** The groups of bones that generate morph targets together. */
+	UPROPERTY()
+	TArray<FNeuralMorphBoneGroup> BoneGroups;
+
+	/** The groups of curves that generate morph targets together. */
+	UPROPERTY()
+	TArray<FNeuralMorphCurveGroup> CurveGroups;
+};
diff --git a/Engine/Plugins/Animation/MLDeformer/NeuralMorphModel/Source/NeuralMorphModel/Public/NeuralMorphModel.h b/Engine/Plugins/Animation/MLDeformer/NeuralMorphModel/Source/NeuralMorphModel/Public/NeuralMorphModel.h
--- a/Engine/Plugins/Animation/MLDeformer/NeuralMorphModel/Source/NeuralMorphModel/Public/NeuralMorphModel.h
+++ b/Engine/Plugins/Animation/MLDeformer/NeuralMorphModel/Source/NeuralMorphModel/Public/NeuralMorphModel.h
@@ -4,6 +4,7 @@
 #include "CoreMinimal.h"
 #include "MLDeformerMorphModel.h"
 #include "GeometryCache.h"
+#include "NeuralMorphTypes.h"
 #include "NeuralMorphModel.generated.h"
 
 class USkeletalMesh;
@@ -31,36 +32,11 @@ struct FExternalMorphSet;
 // Declare our log category.
 NEURALMORPHMODEL_API DECLARE_LOG_CATEGORY_EXTERN(LogNeuralMorphModel, Log, All);
 
-/** 
- * The mode of the model, either local or global. 
- * In local mode the network contains a super simple neural network for each bone, while in 
- * global mode all bones and curves are input to one larger fully connected network.
- * The local mode has faster performance, while global mode can result in higher quality deformation.
- * This model runs its neural network on the CPU, but uses comrpessed GPU based morph targets, which require shader model 5.
- */
-UENUM()
-enum class ENeuralMorphMode : uint8
-{
-	/**
-	 * Each bone creates a set of morph targets and has its own small neural network.
-	 * The local mode can also create more localized morph targets and tends to use slightly less memory.
-	 * This mode is faster to process on the CPU side.
-	 */
-	Local,
-
-	/** 
-	 * There is one fully connected neural network that generates a set of morph targets.
-	 * This has a slightly higher CPU overhead, but could result in higher quality.
-	 * The Global mode is basically the same as the Vertex Delta Model, but runs the neural network on the CPU
-	 * and uses GPU compressed morph targets.
-	 */
-	Global
-};
-
 /**
  * The neural morph model.
  * This generates a set of highly compressed morph targets to approximate a target deformation based on bone rotations and/or curve inputs.
- * The neural network inside this model runs on the CPU and outputs the morph target weights.
+ * During inference, the neural network inside this model runs on the CPU and outputs the morph target weights for the morph targets it generated during training.
+ * Groups of bones and curves can also be defined. Those groups will generate a set of morph targets together as well. This can help when shapes depend on multiple inputs.
  * An external morph target set is generated during training and serialized inside the ML Deformer asset that contains this model.
  * When the ML Deformer component initializes the morph target set is registered. Most of the heavy lifting is done by the UMLDeformerMorphModel class.
  * The neural morph model has two modes: local and global. See the ENeuralMorphMode for a description of the two.
@@ -83,8 +59,11 @@ public:
 	// UMLDeformerModel overrides.
 	virtual FString GetDisplayName() const override			{ return "Neural Morph Model"; }
 	virtual UMLDeformerModelInstance* CreateModelInstance(UMLDeformerComponent* Component) override;
+	virtual UMLDeformerInputInfo* CreateInputInfo();
 	// ~END UMLDeformerModel overrides.
 
+	const TArray<FNeuralMorphBoneGroup>& GetBoneGroups() const		{ return BoneGroups; }
+	const TArray<FNeuralMorphCurveGroup>& GetCurveGroups() const	{ return CurveGroups; }
 	UNeuralMorphNetwork* GetNeuralMorphNetwork() const		{ return NeuralMorphNetwork.Get(); }
 	void SetNeuralMorphNetwork(UNeuralMorphNetwork* Net)	{ NeuralMorphNetwork = Net; }
 	ENeuralMorphMode GetModelMode() const					{ return Mode; }
@@ -96,9 +75,25 @@ public:
 	int32 GetBatchSize() const								{ return BatchSize; }
 	float GetLearningRate() const							{ return LearningRate; }
 	float GetLearningRateDecay() const						{ return LearningRateDecay; }
-	float GetRegularizationFactor() const					{ return RegularizationFactor;  }
+	float GetRegularizationFactor() const					{ return RegularizationFactor; }
 
 public:
+	/**
+	 * The set of bones that are grouped together and generate morph targets together as a whole.
+	 * This can be used in case multiple bones are correlated to each other and work together to produce given shapes.
+	 * Groups are only used when the model is in Local mode.
+	 */
+	UPROPERTY(EditAnywhere, Category = "Inputs and Output")
+	TArray<FNeuralMorphBoneGroup> BoneGroups;
+
+	/**
+	 * The set of curves that are grouped together and generate morph targets together as a whole.
+	 * This can be used in case multiple curves are correlated to each other and work together to produce given shapes.
+	 * Groups are only used when the model is in Local mode.
+	 */
+	UPROPERTY(EditAnywhere, Category = "Inputs and Output")
+	TArray<FNeuralMorphCurveGroup> CurveGroups;
+
 	/**
 	 * The mode that the neural network will operate in. 
 	 * Local mode means there is one tiny network per bone, while global mode has one network for all bones together.
@@ -108,10 +103,10 @@ public:
 	ENeuralMorphMode Mode = ENeuralMorphMode::Local;
 
 	/** 
-	 * The number of morph targets to generate per bone.
+	 * The number of morph targets to generate per bone, curve or group.
 	 * Higher numbers result in better approximation of the target deformation, but also result in a higher memory footprint and slower performance.
 	 */
-	UPROPERTY(EditAnywhere, BlueprintReadWrite, DisplayName = "Num Morph Targets Per Bone", Category = "Training Settings", meta = (ClampMin = "1", ClampMax = "1000"))
+	UPROPERTY(EditAnywhere, BlueprintReadWrite, DisplayName = "Num Morphs Per Bone/Curve/Group", Category = "Training Settings", meta = (ClampMin = "1", ClampMax = "1000"))
 	int32 LocalNumMorphTargetsPerBone = 6;
 
 	/** 
diff --git a/Engine/Plugins/Animation/MLDeformer/NeuralMorphModel/Source/NeuralMorphModel/Public/NeuralMorphModelInstance.h b/Engine/Plugins/Animation/MLDeformer/NeuralMorphModel/Source/NeuralMorphModel/Public/NeuralMorphModelInstance.h
--- a/Engine/Plugins/Animation/MLDeformer/NeuralMorphModel/Source/NeuralMorphModel/Public/NeuralMorphModelInstance.h
+++ b/Engine/Plugins/Animation/MLDeformer/NeuralMorphModel/Source/NeuralMorphModel/Public/NeuralMorphModelInstance.h
@@ -6,6 +6,7 @@
 #include "NeuralMorphModelInstance.generated.h"
 
 class UNeuralMorphNetworkInstance;
+class UNeuralMorphNetwork;
 
 UCLASS()
 class NEURALMORPHMODEL_API UNeuralMorphModelInstance
@@ -21,7 +22,16 @@ public:
 	virtual void Execute(float ModelWeight) override;
 	// ~END UMLDeformerModelInstance overrides.
 
+protected:
+	/**
+	 * Set the network inputs.
+	 */
+	void FillNetworkInputs();
+
 protected:
 	UPROPERTY(Transient)
 	TObjectPtr<UNeuralMorphNetworkInstance> NetworkInstance;
+
+	TArray<int32> BoneGroupIndices;
+	TArray<int32> CurveGroupIndices;
 };
diff --git a/Engine/Plugins/Animation/MLDeformer/NeuralMorphModel/Source/NeuralMorphModel/Public/NeuralMorphNetwork.h b/Engine/Plugins/Animation/MLDeformer/NeuralMorphModel/Source/NeuralMorphModel/Public/NeuralMorphNetwork.h
--- a/Engine/Plugins/Animation/MLDeformer/NeuralMorphModel/Source/NeuralMorphModel/Public/NeuralMorphNetwork.h
+++ b/Engine/Plugins/Animation/MLDeformer/NeuralMorphModel/Source/NeuralMorphModel/Public/NeuralMorphNetwork.h
@@ -8,9 +8,10 @@
 
 class UNeuralMorphNetworkInstance;
 
+
 /** A fully connected layer, which contains the weights and biases for those connections. */ 
 UCLASS()
-class NEURALMORPHMODEL_API UNeuralMorphNetworkLayer
+class NEURALMORPHMODEL_API UNeuralMorphMLPLayer
 	: public UObject
 {
 	GENERATED_BODY()
@@ -24,11 +25,11 @@ public:
 	UPROPERTY()
 	int32 NumOutputs = 0;
 
-	/** The third dimension of the layer. This basically contains the number of bones in local mode. */
+	/** The number of instances of inputs and outputs. */
 	UPROPERTY()
 	int32 Depth = 1;
 
-	/** The weights, which is basically a 2d array. The number of weights will be equal to the number of rows multiplied by number of columns multiplied by the depth. */
+	/** A 2D array of weights. The number of weights equals NumRows x NumColumns x Depth. */
 	UPROPERTY()
 	TArray<float> Weights;
 
@@ -37,10 +38,95 @@ public:
 	TArray<float> Biases;
 };
 
+
+/**
+ * An MLP neural network.
+ */
+UCLASS()
+class NEURALMORPHMODEL_API UNeuralMorphMLP
+	: public UObject
+{
+	GENERATED_BODY()
+
+public:
+	/** Clear the network, getting rid of all weights and biases. */
+	void Empty();
+
+	/** 
+	 * Check if the network is empty or not.
+	 * If it is empty, it means it hasn't been loaded, and cannot do anything.
+	 * Empty means it has no layers.
+	 * @return Returns true if empty, otherwise false is returned.
+	 */
+	bool IsEmpty() const;
+
+	/**
+	 * Load the network from a file on disk.
+	 * When loading fails, the network will be emptied.
+	 * @param Archive The archive to load from.
+	 * @return Returns true when successfully loaded this model, otherwise false is returned.
+	 */
+	bool Load(FArchive& FileReader);
+
+	/**
+	 * Get the number of inputs, which is the number of floats the network takes as input.
+	 * @result The number of input floats to the network.
+	 */
+	int32 GetNumInputs() const;
+
+	/**
+	 * Get the number of outputs, which is the number of floats the network will output.
+	 * @return The number of floats that the network outputs.
+	 */
+	int32 GetNumOutputs() const;
+
+	/**
+	 * Get the number of network layers.
+	 * This equals to the number of hidden layers plus one.
+	 * @return The number of network layers.
+	 */
+	int32 GetNumLayers() const;
+
+	/**
+	 * Get a given network layer.
+	 * @return A reference to the layer, which will contain the weights and biases.
+	 */
+	UNeuralMorphMLPLayer& GetLayer(int32 Index) const;
+
+	/**
+	 * Get the maximum number of inputs for all layers in the network.
+	 * This is used to pre-allocate a temp buffer used to pass the data through the layers.
+	 * @return The maximum number of inputs of all layers.
+	 */
+	int32 GetMaxNumLayerInputs() const;
+
+	/**
+	 * Get the maximum number of outputs for all layers in the network.
+	 * This is used to pre-allocate a temp buffer used to pass the data through the layers.
+	 * @return The maximum number of outputs of all layers.
+	 */
+	int32 GetMaxNumLayerOutputs() const;
+
+private:
+	/**
+	 * Load an MLP fully connected layer from a given archive.
+	 * @param Archive The archive to read from.
+	 */
+	UNeuralMorphMLPLayer* LoadNetworkLayer(FArchive& Archive);
+
+private:
+	/** The network weights and biases of the main network. */
+	UPROPERTY()
+	TArray<TObjectPtr<UNeuralMorphMLPLayer>> Layers;
+};
+
+
 /**
  * The specialized neural network for the Neural Morph Model.
  * This class is used to do inference at runtime at a higher performance than using UNeuralNetwork.
  * The reason why it is faster is because it is a highly specialized network for this specific model.
+ * When the model is a local mode based model, there can be actually two neural networks inside this: one main network and a bone/curve group network.
+ * The group network uses groups of bones or curves to generate morph targets, rather than individual bones/curves.
  */
 UCLASS()
 class NEURALMORPHMODEL_API UNeuralMorphNetwork
@@ -75,14 +161,16 @@ public:
 	UNeuralMorphNetworkInstance* CreateInstance();
 
 	/**
-	 * Get the number of inputs, which is the number of floats the network takes as input.
-	 * @result The number of input floats to the network.
+	 * Get the number of network inputs.
+	 * This is the sum of the number of main network inputs and group network inputs.
+	 * @return The number of inputs.
 	 */
 	int32 GetNumInputs() const;
 
 	/**
-	 * Get the number of outputs, which is the number of floats the network will output.
-	 * @return The number of floats that the network outputs.
+	 * Get the number of network outputs.
+	 * This is the sum of the number of main network outputs and group network outputs.
+	 * @return The number of outputs.
 	 */
 	int32 GetNumOutputs() const;
 
@@ -105,6 +193,22 @@ public:
 	 */
 	int32 GetNumMorphsPerBone() const;
 
+	/**
+	 * Get the number of groups.
+	 * A group is a collection of bones or curves that together generate a number of morph targets.
+	 * Each group must have the same number of items.
+	 * @result The number of groups.
+	 */
+	int32 GetNumGroups() const;
+
+	/**
+	 * Get the number of items per group.
+	 * This indicates how many bones or curves each group has.
+	 * Every group must have the same number of items.
+	 * @return The number of items per group.
+	 */
+	int32 GetNumItemsPerGroup() const;
+
 	/**
 	 * Get the mode that the model was trained for, either global or local mode.
 	 * @return The mode the model was trained for.
@@ -123,29 +227,44 @@ public:
 	 */
 	const TArrayView<const float> GetInputStds() const;
 
+	/** 
+	 * Get the number of floats used to represent a single curve value.
+	 * @return The number of float values per curve.
+	 */
+	int32 GetNumFloatsPerCurve() const;
+
 	/**
-	 * Get the number of network layers.
-	 * This equals to the number of hidden layers plus one.
-	 * @return The number of network layer.
+	 * Get the main MLP.
+	 * This is the MLP (or actually a Multi-MLP) that contains a small network for every bone or curve.
+	 * Each small network will output a set of morph target weights.
+	 * @return A pointer to the main MLP.
 	 */
-	int32 GetNumLayers() const;
+	UNeuralMorphMLP* GetMainMLP() const;
 
 	/**
-	 * Get a given network layer.
-	 * @return A reference to the layer, which will contain the weights and biases.
+	 * Get the MLP used for group of bones and curves.
+	 * This MLP (or actually a Multi-MLP) contains small networks for the bone or curve groups.
+	 * Each group will generate a set of morph target weights.
+	 * @return A pointer to the MLP used for the bone and curve groups.
 	 */
-	UNeuralMorphNetworkLayer& GetLayer(int32 Index) const;
+	UNeuralMorphMLP* GetGroupMLP() const;
 
-	/** 
-	 * Get the number of floats used to represent a single curve value.
-	 * @return The number of float values per curve.
+private:
+	/**
+	 * Load an MLP network from the currently already opened archive.
+	 * @param Archive The archive to load from.
+	 * @return Returns a pointer to the loaded MLP, or a nullptr in case of a failure.
 	 */
-	int32 GetNumFloatsPerCurve() const;
+	UNeuralMorphMLP* LoadMLP(FArchive& FileReader);
 
 private:
-	/** The network weights and biases, between the different layers. */
+	/** The MLP that acts as main network. */
+	UPROPERTY()
+	TObjectPtr<UNeuralMorphMLP> MainMLP;
+
+	/** The MLP for the bone and curve groups, when in local mode. */
 	UPROPERTY()
-	TArray<TObjectPtr<UNeuralMorphNetworkLayer>> Layers;
+	TObjectPtr<UNeuralMorphMLP> GroupMLP;
 
 	/** The means of the input values, used to normalize inputs. */
 	UPROPERTY()
@@ -174,6 +293,14 @@ private:
 	/** The number of floats per curve. */
 	UPROPERTY()
 	int32 NumFloatsPerCurve = 1;
+
+	/** The number of groups. This is a group of bones/curves that create their own morphs. */
+	UPROPERTY()
+	int32 NumGroups = 0;
+
+	/** The number of items (bones/curves) per group. */
+	UPROPERTY()
+	int32 NumItemsPerGroup = 0;
 };
 
 /** 
@@ -192,11 +319,18 @@ class NEURALMORPHMODEL_API UNeuralMorphNetworkInstance
 public:
 	/**
 	 * Get the network input buffer.
-	 * @return The array view of floats that represents the network inputs.
+	 * @return The array view of floats that represents the main network inputs.
 	 */
 	TArrayView<float> GetInputs();
 	TArrayView<const float> GetInputs() const;
 
+	/**
+	 * Get the inputs for the group MLP.
+	 * @return The array view of floats that represent the group network inputs.
+	 */
+	TArrayView<float> GetGroupInputs();
+	TArrayView<const float> GetGroupInputs() const;
+
 	/**
 	 * Get the network output buffer.
 	 * @return The array view of floats that represents the network outputs.
@@ -220,12 +354,10 @@ public:
 private:
 	struct FRunSettings
 	{
-		float* TempInputBuffer;
-		float* TempOutputBuffer;
-		const float* InputStdsBuffer;
-		const float* InputMeansBuffer;
-		const float* InputBuffer;
-		float* OutputBuffer;
+		float* TempInputBuffer = nullptr;
+		float* TempOutputBuffer = nullptr;
+		const float* InputBuffer = nullptr;
+		float* OutputBuffer = nullptr;
 	};
 
 	/**
@@ -235,10 +367,11 @@ private:
 	void RunGlobalModel(const FRunSettings& RunSettings);
 
 	/**
- 	 * Run inference using the local model mode, which basically has a small MLP for every bone.
+ 	 * Run a local mlp, which internally is acting like a multi-mlp (one mlp per bone or curve).
+	 * @param MLP The MLP to run.
 	 * @param RunSettings The settings used to run the model. This specifies a set of buffers.
 	 */
-	void RunLocalModel(const FRunSettings& RunSettings);
+	void RunLocalMLP(UNeuralMorphMLP& MLP, const FRunSettings& RunSettings);
 
 	/**
  	 * Initialize this instance using a given network object.
@@ -248,9 +381,12 @@ private:
 	void Init(UNeuralMorphNetwork* InNeuralNetwork);
 
 private:
-	/** The input values. */
+	/** The input values for the main MLP. */
 	TArray<float> Inputs;
 
+	/** The input values for the group MLP. */
+	TArray<float> GroupInputs;
+
 	/** The output values. */
 	TArray<float> Outputs;
 
diff --git a/Engine/Plugins/Animation/MLDeformer/NeuralMorphModel/Source/NeuralMorphModel/Public/NeuralMorphTypes.h b/Engine/Plugins/Animation/MLDeformer/NeuralMorphModel/Source/NeuralMorphModel/Public/NeuralMorphTypes.h
new file mode 100644
--- /dev/null
+++ b/Engine/Plugins/Animation/MLDeformer/NeuralMorphModel/Source/NeuralMorphModel/Public/NeuralMorphTypes.h
@@ -0,0 +1,65 @@
+// Copyright Epic Games, Inc. All Rights Reserved.
+#pragma once
+
+#include "CoreTypes.h"
+#include "Containers/Array.h"
+#include "MLDeformerCurveReference.h"
+#include "NeuralMorphTypes.generated.h"
+
+struct FBoneReference;
+
+/** 
+ * The mode of the model, either local or global. 
+ * In local mode the network contains a super simple neural network for each bone, while in 
+ * global mode all bones and curves are input to one larger fully connected network.
+ * The local mode has faster performance, while global mode can result in higher quality deformation.
+ * This model runs its neural network on the CPU, but uses comrpessed GPU based morph targets, which require shader model 5.
+ */
+UENUM()
+enum class ENeuralMorphMode : uint8
+{
+	/**
+	 * Each bone creates a set of morph targets and has its own small neural network.
+	 * The local mode can also create more localized morph targets and tends to use slightly less memory.
+	 * This mode is faster to process on the CPU side.
+	 */
+	Local,
+
+	/** 
+	 * There is one fully connected neural network that generates a set of morph targets.
+	 * This has a slightly higher CPU overhead, but could result in higher quality.
+	 * The Global mode is basically the same as the Vertex Delta Model, but runs the neural network on the CPU
+	 * and uses GPU compressed morph targets.
+	 */
+	Global
+};
+
+/**
+ * A group of bones, which can generate morph targets together.
+ * This is useful when there are specific correlations between different bones, that all effect the same region on the body.
+ */
+USTRUCT()
+struct NEURALMORPHMODEL_API FNeuralMorphBoneGroup
+{
+	GENERATED_BODY()
+
+public:
+	/** The list of bones that should form a group together. */
+	UPROPERTY(EditAnywhere, Category = "Bones")
+	TArray<FBoneReference> BoneNames;
+};
+
+/**
+ * A group of curves, which can generate morph targets together.
+ * This is useful when there are specific correlations between different curves, that all effect the same region on the body.
+ */
+USTRUCT()
+struct NEURALMORPHMODEL_API FNeuralMorphCurveGroup
+{
+	GENERATED_BODY()
+
+public:
+	/** The list of curves that should form a group together. */
+	UPROPERTY(EditAnywhere, Category = "Curves")
+	TArray<FMLDeformerCurveReference> CurveNames;
+};
diff --git a/Engine/Plugins/Animation/MLDeformer/NeuralMorphModel/Source/NeuralMorphModelEditor/Private/NeuralMorphEditorModel.cpp b/Engine/Plugins/Animation/MLDeformer/NeuralMorphModel/Source/NeuralMorphModelEditor/Private/NeuralMorphEditorModel.cpp
--- a/Engine/Plugins/Animation/MLDeformer/NeuralMorphModel/Source/NeuralMorphModelEditor/Private/NeuralMorphEditorModel.cpp
+++ b/Engine/Plugins/Animation/MLDeformer/NeuralMorphModel/Source/NeuralMorphModelEditor/Private/NeuralMorphEditorModel.cpp
@@ -4,9 +4,12 @@
 #include "IDetailsView.h"
 #include "NeuralMorphModel.h"
 #include "NeuralMorphNetwork.h"
-#include "MLDeformerMorphModelVizSettings.h"
+#include "NeuralMorphInputInfo.h"
 #include "NeuralMorphTrainingModel.h"
+#include "MLDeformerMorphModelVizSettings.h"
 #include "MLDeformerEditorToolkit.h"
+#include "Engine/SkeletalMesh.h"
+#include "Animation/Skeleton.h"
 
 #define LOCTEXT_NAMESPACE "NeuralMorphEditorModel"
 
@@ -30,12 +33,20 @@ namespace UE::NeuralMorphModel
 		// Process the base class property changes.
 		FMLDeformerMorphModelEditorModel::OnPropertyChanged(PropertyChangedEvent);
 
-		// Handle property changes of this model.
-		if (Property->GetFName() == GET_MEMBER_NAME_CHECKED(UNeuralMorphModel, Mode))
+		if (Property->GetFName() == GET_MEMBER_NAME_CHECKED(UNeuralMorphModel, BoneGroups) ||
+			Property->GetFName() == GET_MEMBER_NAME_CHECKED(UNeuralMorphModel, CurveGroups) ||
+			Property->GetFName() == TEXT("BoneNames") ||	// FNeuralMorphBoneGroup::BoneNames has items added or removed (it is an array).
+			Property->GetFName() == TEXT("CurveNames") ||	// FNeuralMorphCurveGroup::CurveNames has items added or removed (it is an array).
+			Property->GetFName() == TEXT("BoneName") ||		// The bone name inside one of the items in the BoneNames list changed.
+			Property->GetFName() == TEXT("CurveName"))		// The curve name inside one of the items in the CurveNames list changed.
+		{
+			UpdateIsReadyForTrainingState();
+		}
+		else if (Property->GetFName() == GET_MEMBER_NAME_CHECKED(UNeuralMorphModel, Mode))
 		{
 			if (PropertyChangedEvent.ChangeType == EPropertyChangeType::ValueSet)
 			{
-				SetResamplingInputOutputsNeeded(true);	// Needed as local mode doesn't support curves, while the global mode does.
+				UpdateIsReadyForTrainingState();
 				GetEditor()->GetModelDetailsView()->ForceRefresh();
 			}
 		}
@@ -46,7 +57,18 @@ namespace UE::NeuralMorphModel
 #if NEURALMORPHMODEL_FORCE_USE_NNI
 		return GetMorphModel()->GetNeuralNetwork() != nullptr;
 #else
-		return (GetMorphModel()->GetNeuralNetwork() != nullptr) || (GetNeuralMorphModel()->GetNeuralMorphNetwork() != nullptr);
+		if (GetMorphModel()->GetNeuralNetwork() != nullptr)
+		{
+			return true;
+		}
+
+		UNeuralMorphNetwork* MorphNetwork = GetNeuralMorphModel()->GetNeuralMorphNetwork();
+		if (MorphNetwork != nullptr && MorphNetwork->GetMainMLP() != nullptr)
+		{
+			return true;
+		}
+
+		return false;
 #endif
 	}
 
@@ -76,7 +98,7 @@ namespace UE::NeuralMorphModel
 		if (!NeuralNet->Load(NetworkFilename))
 		{
 			UE_LOG(LogNeuralMorphModel, Error, TEXT("Failed to load neural morph network from file '%s'!"), *NetworkFilename);
-			NeuralNet = nullptr;
+			NeuralNet->ConditionalBeginDestroy();
 
 			// Restore the deltas to the ones before training started.
 			GetMorphModel()->SetMorphTargetDeltas(MorphTargetDeltasBackup);
@@ -85,6 +107,7 @@ namespace UE::NeuralMorphModel
 
 		if (NeuralNet->IsEmpty())
 		{
+			NeuralNet->ConditionalBeginDestroy();
 			NeuralNet = nullptr;
 		}
 
@@ -93,6 +116,147 @@ namespace UE::NeuralMorphModel
 #endif
 		return true;
 	}
+
+	void FNeuralMorphEditorModel::InitInputInfo(UMLDeformerInputInfo* InputInfo)
+	{
+		FMLDeformerEditorModel::InitInputInfo(InputInfo);
+
+		UNeuralMorphInputInfo* NeuralInputInfo = Cast<UNeuralMorphInputInfo>(InputInfo);
+		if (NeuralInputInfo == nullptr)
+		{
+			return;
+		}
+
+		UNeuralMorphModel* NeuralMorphModel = Cast<UNeuralMorphModel>(Model);
+		if (NeuralMorphModel->GetModelMode() != ENeuralMorphMode::Local)
+		{
+			return;
+		}
+
+		const bool bIncludeBones = Model->DoesSupportBones() && Model->ShouldIncludeBonesInTraining();
+		const bool bIncludeCurves = Model->DoesSupportCurves() && Model->ShouldIncludeCurvesInTraining();
+		const USkeleton* Skeleton = Model->GetSkeletalMesh() ? Model->GetSkeletalMesh()->GetSkeleton() : nullptr;
+		const USkeletalMesh* SkeletalMesh = Model->GetSkeletalMesh();
+
+		// Handle bones.
+		if (bIncludeBones && SkeletalMesh)
+		{
+			const TArray<FNeuralMorphBoneGroup>& ModelBoneGroups = NeuralMorphModel->GetBoneGroups();
+			const FReferenceSkeleton& RefSkeleton = SkeletalMesh->GetRefSkeleton();
+			for (int32 BoneGroupIndex = 0; BoneGroupIndex < ModelBoneGroups.Num(); ++BoneGroupIndex)
+			{
+				const FNeuralMorphBoneGroup& BoneGroup = ModelBoneGroups[BoneGroupIndex];
+				if (BoneGroup.BoneNames.IsEmpty())
+				{
+					continue;
+				}
+
+				NeuralInputInfo->GetBoneGroups().AddDefaulted();
+				FNeuralMorphBoneGroup& NewGroup = NeuralInputInfo->GetBoneGroups().Last();
+
+				const int32 NumBonesInGroup = BoneGroup.BoneNames.Num();
+				NewGroup.BoneNames.AddDefaulted(NumBonesInGroup);
+				for (int32 Index = 0; Index < NumBonesInGroup; ++Index)
+				{
+					const FBoneReference& BoneRef = BoneGroup.BoneNames[Index];
+					if (!BoneRef.BoneName.IsValid() || BoneRef.BoneName.IsNone())
+					{
+						UE_LOG(LogNeuralMorphModel, Warning, TEXT("Invalid or 'None' bone detected inside curve group %d, please fix this."), BoneGroupIndex);
+						continue;
+					}
+
+					if (RefSkeleton.FindBoneIndex(BoneRef.BoneName) == INDEX_NONE)
+					{
+						UE_LOG(LogNeuralMorphModel, Warning, TEXT("Bone '%s' inside bone group %d doesn't exist, please fix this."), *BoneRef.BoneName.ToString(), BoneGroupIndex);
+						continue;
+					}
+
+					if (!InputInfo->GetBoneNames().Contains(BoneRef.BoneName))
+					{
+						UE_LOG(LogNeuralMorphModel, Warning, TEXT("Bone '%s' inside bone group %d isn't included in the bone list that are input to the model."), *BoneRef.BoneName.ToString(), BoneGroupIndex);
+						continue;
+					}
+
+					NewGroup.BoneNames[Index] = BoneRef.BoneName;
+				}
+			}
+		}
+
+		// Handle curves.
+		if (bIncludeCurves && SkeletalMesh)
+		{
+			const FSmartNameMapping* SmartNameMapping = Skeleton ? Skeleton->GetSmartNameContainer(USkeleton::AnimCurveMappingName) : nullptr;
+			if (SmartNameMapping) // When there are curves.
+			{
+				const TArray<FNeuralMorphCurveGroup>& ModelCurveGroups = NeuralMorphModel->GetCurveGroups();
+				const FReferenceSkeleton& RefSkeleton = SkeletalMesh->GetRefSkeleton();
+				for (int32 CurveGroupIndex = 0; CurveGroupIndex < ModelCurveGroups.Num(); ++CurveGroupIndex)
+				{				
+					const FNeuralMorphCurveGroup& CurveGroup = ModelCurveGroups[CurveGroupIndex];
+					if (CurveGroup.CurveNames.IsEmpty())
+					{
+						continue;
+					}
+
+					NeuralInputInfo->GetCurveGroups().AddDefaulted();
+					FNeuralMorphCurveGroup& NewGroup = NeuralInputInfo->GetCurveGroups().Last();
+
+					const int32 NumCurvesInGroup = CurveGroup.CurveNames.Num();
+					NewGroup.CurveNames.AddDefaulted(NumCurvesInGroup);
+					for (int32 Index = 0; Index < NumCurvesInGroup; ++Index)
+					{
+						const FMLDeformerCurveReference& CurveRef = CurveGroup.CurveNames[Index];
+						if (!CurveRef.CurveName.IsValid() || CurveRef.CurveName.IsNone())
+						{
+							UE_LOG(LogNeuralMorphModel, Warning, TEXT("Invalid or 'None' curve detected inside curve group %d, please fix this."), CurveGroupIndex);
+							continue;
+						}
+
+						if (!SmartNameMapping->Exists(CurveRef.CurveName))
+						{
+							UE_LOG(LogNeuralMorphModel, Warning, TEXT("Curve '%s' inside curve group %d doesn't exist, please fix this."), *CurveRef.CurveName.ToString(), CurveGroupIndex);
+							continue;
+						}
+
+						if (!InputInfo->GetCurveNames().Contains(CurveRef.CurveName))
+						{
+							UE_LOG(LogNeuralMorphModel, Warning, TEXT("Curve '%s' inside curve group %d isn't included in the curve list that are input to the model."), *CurveRef.CurveName.ToString(), CurveGroupIndex);
+							continue;
+						}
+
+						NewGroup.CurveNames[Index] = CurveRef.CurveName;
+					}
+				}
+			} // if SmartNameMapping
+		}
+	}
+
+	void FNeuralMorphEditorModel::UpdateIsReadyForTrainingState()
+	{
+		FMLDeformerMorphModelEditorModel::UpdateIsReadyForTrainingState();
+		UNeuralMorphInputInfo* Info = Cast<UNeuralMorphInputInfo>(GetEditorInputInfo());
+		if (Info)
+		{
+			if (GetNeuralMorphModel()->GetModelMode() == ENeuralMorphMode::Local)
+			{
+				bIsReadyForTraining &= !Info->HasInvalidGroups();
+			}
+		}
+	}
+
+	FText FNeuralMorphEditorModel::GetOverlayText() const
+	{
+		FText Text = FMLDeformerMorphModelEditorModel::GetOverlayText();
+		UNeuralMorphInputInfo* Info = Cast<UNeuralMorphInputInfo>(GetEditorInputInfo());
+		if (GetNeuralMorphModel()->GetModelMode() == ENeuralMorphMode::Local && Info->HasInvalidGroups())
+		{
+			Text = FText::Format(
+				LOCTEXT("GroupErrorFormat", "{0}\n{1}"), 
+				Text,
+				LOCTEXT("GroupErrorText", "Bone and curve groups require at least two valid items.\nCheck the log warnings for more information."));
+		}
+		return Text;
+	}
 }	// namespace UE::NeuralMorphModel
 
 #undef LOCTEXT_NAMESPACE
diff --git a/Engine/Plugins/Animation/MLDeformer/NeuralMorphModel/Source/NeuralMorphModelEditor/Private/NeuralMorphModelDetails.cpp b/Engine/Plugins/Animation/MLDeformer/NeuralMorphModel/Source/NeuralMorphModelEditor/Private/NeuralMorphModelDetails.cpp
--- a/Engine/Plugins/Animation/MLDeformer/NeuralMorphModel/Source/NeuralMorphModelEditor/Private/NeuralMorphModelDetails.cpp
+++ b/Engine/Plugins/Animation/MLDeformer/NeuralMorphModel/Source/NeuralMorphModelEditor/Private/NeuralMorphModelDetails.cpp
@@ -2,6 +2,8 @@
 
 #include "NeuralMorphModelDetails.h"
 #include "NeuralMorphModel.h"
+#include "NeuralMorphInputInfo.h"
+#include "NeuralMorphEditorModel.h"
 #include "DetailCategoryBuilder.h"
 #include "DetailWidgetRow.h"
 
@@ -54,8 +56,24 @@ namespace UE::NeuralMorphModel
 		TrainingSettingsCategoryBuilder->AddProperty(GET_MEMBER_NAME_CHECKED(UNeuralMorphModel, RegularizationFactor), UNeuralMorphModel::StaticClass());
 	}
 
+	void FNeuralMorphModelDetails::AddTrainingInputFilters()
+	{
+		FMLDeformerMorphModelDetails::AddTrainingInputFilters();
+
+		UNeuralMorphModel* NeuralMorphModel = Cast<UNeuralMorphModel>(Model);
+		check(NeuralMorphModel);
+
+		InputOutputCategoryBuilder->AddProperty(GET_MEMBER_NAME_CHECKED(UNeuralMorphModel, BoneGroups), UNeuralMorphModel::StaticClass())
+			.Visibility(NeuralMorphModel->Mode == ENeuralMorphMode::Local ? EVisibility::Visible : EVisibility::Collapsed);
+
+		InputOutputCategoryBuilder->AddProperty(GET_MEMBER_NAME_CHECKED(UNeuralMorphModel, CurveGroups), UNeuralMorphModel::StaticClass())
+			.Visibility(NeuralMorphModel->Mode == ENeuralMorphMode::Local ? EVisibility::Visible : EVisibility::Collapsed);
+	}
+
 	void FNeuralMorphModelDetails::AddTrainingSettingsErrors()
 	{
+		FMLDeformerMorphModelDetails::AddTrainingSettingsErrors();
+
 		UNeuralMorphModel* NeuralMorphModel = Cast<UNeuralMorphModel>(Model);
 		check(NeuralMorphModel);
 
diff --git a/Engine/Plugins/Animation/MLDeformer/NeuralMorphModel/Source/NeuralMorphModelEditor/Private/NeuralMorphTrainingModel.cpp b/Engine/Plugins/Animation/MLDeformer/NeuralMorphModel/Source/NeuralMorphModelEditor/Private/NeuralMorphTrainingModel.cpp
new file mode 100644
--- /dev/null
+++ b/Engine/Plugins/Animation/MLDeformer/NeuralMorphModel/Source/NeuralMorphModelEditor/Private/NeuralMorphTrainingModel.cpp
@@ -0,0 +1,54 @@
+// Copyright Epic Games, Inc. All Rights Reserved.
+
+#include "NeuralMorphTrainingModel.h"
+#include "NeuralMorphEditorModel.h"
+#include "NeuralMorphModel.h"
+#include "NeuralMorphInputInfo.h"
+
+
+bool UNeuralMorphTrainingModel::ForceUseNNI() const
+{
+	#if NEURALMORPHMODEL_FORCE_USE_NNI
+		return true;
+	#else
+		return false;
+	#endif
+}
+
+int32 UNeuralMorphTrainingModel::GetNumBoneGroups() const
+{
+	UNeuralMorphInputInfo* InputInfo = Cast<UNeuralMorphInputInfo>(EditorModel->GetEditorInputInfo());
+	return InputInfo ? InputInfo->GetBoneGroups().Num() : 0;
+}
+
+int32 UNeuralMorphTrainingModel::GetNumCurveGroups() const
+{
+	UNeuralMorphInputInfo* InputInfo = Cast<UNeuralMorphInputInfo>(EditorModel->GetEditorInputInfo());	
+	return InputInfo ? InputInfo->GetCurveGroups().Num() : 0;
+}
+
+TArray<int32> UNeuralMorphTrainingModel::GenerateBoneGroupIndices() const
+{
+	UNeuralMorphInputInfo* InputInfo = Cast<UNeuralMorphInputInfo>(EditorModel->GetEditorInputInfo());
+	if (InputInfo == nullptr)
+	{
+		return TArray<int32>();
+	}
+
+	TArray<int32> OutBoneGroupIndices;
+	InputInfo->GenerateBoneGroupIndices(OutBoneGroupIndices);
+	return MoveTemp(OutBoneGroupIndices);
+}
+
+TArray<int32> UNeuralMorphTrainingModel::GenerateCurveGroupIndices() const
+{
+	UNeuralMorphInputInfo* InputInfo = Cast<UNeuralMorphInputInfo>(EditorModel->GetEditorInputInfo());
+	if (InputInfo == nullptr)
+	{
+		return TArray<int32>();
+	}
+
+	TArray<int32> OutCurveGroupIndices;
+	InputInfo->GenerateCurveGroupIndices(OutCurveGroupIndices);
+	return MoveTemp(OutCurveGroupIndices);
+}
diff --git a/Engine/Plugins/Animation/MLDeformer/NeuralMorphModel/Source/NeuralMorphModelEditor/Public/NeuralMorphEditorModel.h b/Engine/Plugins/Animation/MLDeformer/NeuralMorphModel/Source/NeuralMorphModelEditor/Public/NeuralMorphEditorModel.h
--- a/Engine/Plugins/Animation/MLDeformer/NeuralMorphModel/Source/NeuralMorphModelEditor/Public/NeuralMorphEditorModel.h
+++ b/Engine/Plugins/Animation/MLDeformer/NeuralMorphModel/Source/NeuralMorphModelEditor/Public/NeuralMorphEditorModel.h
@@ -25,11 +25,14 @@ namespace UE::NeuralMorphModel
 		// FGCObject overrides.
 		virtual FString GetReferencerName() const override		{ return TEXT("FNeuralMorphEditorModel"); }
 		virtual bool LoadTrainedNetwork() const override;
+		virtual void UpdateIsReadyForTrainingState() override;
 		// ~END FGCObject overrides.
 	
 		// FMLDeformerEditorModel overrides.
 		virtual ETrainingResult Train() override;
 		virtual bool IsTrained() const override;
+		virtual void InitInputInfo(UMLDeformerInputInfo* InputInfo) override;
+		virtual FText GetOverlayText() const override;
 		virtual void OnPropertyChanged(FPropertyChangedEvent& PropertyChangedEvent) override;
 		// ~END FMLDeformerEditorModel overrides.
 
diff --git a/Engine/Plugins/Animation/MLDeformer/NeuralMorphModel/Source/NeuralMorphModelEditor/Public/NeuralMorphModelDetails.h b/Engine/Plugins/Animation/MLDeformer/NeuralMorphModel/Source/NeuralMorphModelEditor/Public/NeuralMorphModelDetails.h
--- a/Engine/Plugins/Animation/MLDeformer/NeuralMorphModel/Source/NeuralMorphModelEditor/Public/NeuralMorphModelDetails.h
+++ b/Engine/Plugins/Animation/MLDeformer/NeuralMorphModel/Source/NeuralMorphModelEditor/Public/NeuralMorphModelDetails.h
@@ -26,6 +26,7 @@ namespace UE::NeuralMorphModel
 
 		// FMLDeformerModelDetails overrides.
 		virtual void AddTrainingSettingsErrors() override;
+		virtual void AddTrainingInputFilters() override;
 		// ~END FMLDeformerModelDetails overrides.
 	};
 }	// namespace UE::NeuralMorphModel
diff --git a/Engine/Plugins/Animation/MLDeformer/NeuralMorphModel/Source/NeuralMorphModelEditor/Public/NeuralMorphTrainingModel.h b/Engine/Plugins/Animation/MLDeformer/NeuralMorphModel/Source/NeuralMorphModelEditor/Public/NeuralMorphTrainingModel.h
--- a/Engine/Plugins/Animation/MLDeformer/NeuralMorphModel/Source/NeuralMorphModelEditor/Public/NeuralMorphTrainingModel.h
+++ b/Engine/Plugins/Animation/MLDeformer/NeuralMorphModel/Source/NeuralMorphModelEditor/Public/NeuralMorphTrainingModel.h
@@ -2,9 +2,9 @@
 
 #pragma once
 
-#include "CoreMinimal.h"
+#include "CoreTypes.h"
+#include "Containers/Array.h"
 #include "MLDeformerTrainingModel.h"
-#include "NeuralMorphModel.h"
 #include "NeuralMorphTrainingModel.generated.h"
 
 /**
@@ -30,13 +30,18 @@ public:
 	 * Do we want to force using NNI for inference, instead of our custom inference?
 	 * @result Returns true when NNI should be used for inference, or false if custom (faster) inference should be used.
 	 */
-	UFUNCTION(BlueprintPure, Category = "Training Data")
-	bool ForceUseNNI() const
-	{
-		#if NEURALMORPHMODEL_FORCE_USE_NNI
-			return true;
-		#else
-			return false;
-		#endif
-	}
+	UFUNCTION(BlueprintPure, Category = "Training Model")
+	bool ForceUseNNI() const;
+
+	UFUNCTION(BlueprintPure, Category = "Training Model")
+	int32 GetNumBoneGroups() const;
+
+	UFUNCTION(BlueprintPure, Category = "Training Model")
+	int32 GetNumCurveGroups() const;
+
+	UFUNCTION(BlueprintPure, Category = "Training Model")
+	TArray<int32> GenerateBoneGroupIndices() const;
+
+	UFUNCTION(BlueprintPure, Category = "Training Model")
+	TArray<int32> GenerateCurveGroupIndices() const;
 };
diff --git a/Engine/Source/Runtime/Engine/Private/Components/SkinnedMeshComponent.cpp b/Engine/Source/Runtime/Engine/Private/Components/SkinnedMeshComponent.cpp
--- a/Engine/Source/Runtime/Engine/Private/Components/SkinnedMeshComponent.cpp
+++ b/Engine/Source/Runtime/Engine/Private/Components/SkinnedMeshComponent.cpp
@@ -353,7 +353,8 @@ void FExternalMorphSetWeights::UpdateNumActiveMorphTargets()
 
 void FExternalMorphSetWeights::ZeroWeights(bool bZeroNumActiveMorphTargets)
 {
-	for (int32 Index = 0; Index < Weights.Num(); ++Index)
+	const int32 NumWeights = Weights.Num();
+	for (int32 Index = 0; Index < NumWeights; ++Index)
 	{
 		Weights[Index] = 0.0f;
 	}
