[MLDeformer] Fix a compile issue.

#trivial
#rnx

diff --git a/Engine/Plugins/Animation/MLDeformer/MLDeformerFramework/Source/MLDeformerFrameworkEditor/Private/MLDeformerEditorModel.cpp b/Engine/Plugins/Animation/MLDeformer/MLDeformerFramework/Source/MLDeformerFrameworkEditor/Private/MLDeformerEditorModel.cpp
--- a/Engine/Plugins/Animation/MLDeformer/MLDeformerFramework/Source/MLDeformerFrameworkEditor/Private/MLDeformerEditorModel.cpp
+++ b/Engine/Plugins/Animation/MLDeformer/MLDeformerFramework/Source/MLDeformerFrameworkEditor/Private/MLDeformerEditorModel.cpp
@@ -35,6 +35,7 @@
 #include "GeometryCacheMeshData.h"
 #include "GeometryCacheTrack.h"
 #include "NeuralNetwork.h"
+#include "BoneContainer.h"
 
 #define LOCTEXT_NAMESPACE "MLDeformerEditorModel"
 
diff --git a/Engine/Plugins/Animation/MLDeformer/NeuralMorphModel/Source/NeuralMorphModel/Public/NeuralMorphTypes.h b/Engine/Plugins/Animation/MLDeformer/NeuralMorphModel/Source/NeuralMorphModel/Public/NeuralMorphTypes.h
--- a/Engine/Plugins/Animation/MLDeformer/NeuralMorphModel/Source/NeuralMorphModel/Public/NeuralMorphTypes.h
+++ b/Engine/Plugins/Animation/MLDeformer/NeuralMorphModel/Source/NeuralMorphModel/Public/NeuralMorphTypes.h
@@ -3,11 +3,10 @@
 
 #include "CoreTypes.h"
 #include "Containers/Array.h"
+#include "BoneContainer.h"
 #include "MLDeformerCurveReference.h"
 #include "NeuralMorphTypes.generated.h"
 
-struct FBoneReference;
-
 /** 
  * The mode of the model, either local or global. 
  * In local mode the network contains a super simple neural network for each bone, while in 
diff --git a/Engine/Plugins/Animation/MLDeformer/NeuralMorphModel/Source/NeuralMorphModelEditor/Private/NeuralMorphEditorModel.cpp b/Engine/Plugins/Animation/MLDeformer/NeuralMorphModel/Source/NeuralMorphModelEditor/Private/NeuralMorphEditorModel.cpp
--- a/Engine/Plugins/Animation/MLDeformer/NeuralMorphModel/Source/NeuralMorphModelEditor/Private/NeuralMorphEditorModel.cpp
+++ b/Engine/Plugins/Animation/MLDeformer/NeuralMorphModel/Source/NeuralMorphModelEditor/Private/NeuralMorphEditorModel.cpp
@@ -10,6 +10,7 @@
 #include "MLDeformerEditorToolkit.h"
 #include "Engine/SkeletalMesh.h"
 #include "Animation/Skeleton.h"
+#include "BoneContainer.h"
 
 #define LOCTEXT_NAMESPACE "NeuralMorphEditorModel"
 
