[MLDeformer]

diff --git a/Engine/Plugins/Animation/MLDeformer/MLDeformerFramework/Source/MLDeformerFrameworkEditor/Public/MLDeformerMorphModelEditorModel.h b/Engine/Plugins/Animation/MLDeformer/MLDeformerFramework/Source/MLDeformerFrameworkEditor/Public/MLDeformerMorphModelEditorModel.h
--- a/Engine/Plugins/Animation/MLDeformer/MLDeformerFramework/Source/MLDeformerFrameworkEditor/Public/MLDeformerMorphModelEditorModel.h
+++ b/Engine/Plugins/Animation/MLDeformer/MLDeformerFramework/Source/MLDeformerFrameworkEditor/Public/MLDeformerMorphModelEditorModel.h
@@ -6,6 +6,7 @@
 #include "MLDeformerGeomCacheEditorModel.h"
 #include "MLDeformerGeomCacheSampler.h"
 #include "MLDeformerEditorActor.h"
+#include "Rendering/ColorVertexBuffer.h"
 
 class UMLDeformerMorphModel;
 class UMLDeformerMorphModelVizSettings;
@@ -34,6 +35,95 @@ namespace UE::MLDeformer
 		virtual void OnPostTraining(ETrainingResult TrainingResult, bool bUsePartiallyTrainedWhenAborted) override;
 		// ~END FMLDeformerEditorModel overrides.
 
+		/**
+		 * Get the mask buffer for a given morph target.
+		 * @param MorphTargetIndex The morph target index, excluding the first 'means' morph target that we always have, so 0 would mean the actual first real morph target.
+		 */
+		virtual const TArrayView<const float> GetMaskForMorphTarget(int32 MorphTargetIndex) const;
+
+		/**
+		 * Check if input masking is supported in this model.
+		 * This allows bones and curves to specify a given mask. These masks can be used during training.
+		 * On default this is disabled.
+		 * @return Returns true if supported, or false otherwise.
+		 */
+		virtual bool IsInputMaskingSupported() const;
+
+		/**
+		 * Fill the mask buffer for a given item with a specific value.
+		 * An item can be a bone, curve, bone group or curve group.
+		 * The number of values written is always Model->GetNumBaseMeshVerts().
+		 * @param MaskBuffer The mask buffer to write to. The mask buffer view has Model->GetBaseNumVerts() elements.
+		 * @param Value The value to write for all vertices.
+		 */
+		void FillMaskValues(TArrayView<float> ItemMaskBuffer, float Value) const;
+
+		/**
+		 * Add bone influences to the mask.
+		 * Every vertex that is linked to the specified bone will write its skinning weight to the mask buffer.
+		 * @param SkeletonBoneIndex The bone index inside our reference skeleton.
+		 * @param MaskBuffer The mask buffer to write to.
+		 */
+		void ApplyBoneToMask(int32 SkeletonBoneIndex, TArrayView<float> MaskBuffer);
+
+		/**
+		 * Recursively add bone masks for this bone and its parents.
+		 * The number of parents is controlled by the MaxHierarchyDepth.
+		 * @param RefSkel The reference skeleton.
+		 * @param SkeletonBoneIndex The bone index of our current bone to add to the mask.
+		 * @param MaxHierarchyDepth The maximum hierarchy depth. If this is 1, it means that we visit only up to the parent of the specified bone. If it is 2, we also visit the parent of the parent, etc.
+		 * @param OutBoneAdded The list of bones that we already added to the mask. Bones that are in this list will not be added to the mask again.
+		 * @param VirtualParentTable The table of virtual parents. This lets us know that we also have to add say the finger bone vertices to the mask of the hand in our example described above.
+		 * @param CurHierarchyDepth The current recursive hierarchy depth. This is used to track how deep we are when moving up the hierarchy. A value of 1 means we are currently adding the parent of the initial bone.
+		 */
+		void RecursiveAddBoneToMaskUpwards(const FReferenceSkeleton& RefSkel, int32 SkeletonBoneIndex, int32 MaxHierarchyDepth, const TArray<int32>& VirtualParentTable, TArray<int32>& OutBonesAdded, int32 CurHierarchyDepth = 0);
+
+		/**
+		 * Recursively add bone masks for this bone and its children.
+		 * The number of child steps down the hierarchy is controlled by the MaxHierarchyDepth.
+		 * @param RefSkel The reference skeleton.
+		 * @param SkeletonBoneIndex The bone index of our current bone to add to the mask.
+		 * @param MaxHierarchyDepth The maximum hierarchy depth. If this is 1, it means that we visit only down to the children of the specified bone. If it is 2, we also visit the children of the children, etc.
+		 * @param VirtualParentTable The table of virtual parents. This lets us know that we also have to add say the finger bone vertices to the mask of the hand in our example described above.
+		 * @param OutBoneAdded The list of bones that we already added to the mask. Bones that are in this list will not be added to the mask again.
+		 * @param CurHierarchyDepth The current recursive hierarchy depth. This is used to track how deep we are when moving down the hierarchy. A value of 1 means we are currently adding a child node of the initial bone.
+		 */
+		void RecursiveAddBoneToMaskDownwards(const FReferenceSkeleton& RefSkel, int32 SkeletonBoneIndex, int32 MaxHierarchyDepth, const TArray<int32>& VirtualParentTable, TArray<int32>& OutBonesAdded, int32 CurHierarchyDepth = 0);
+
+		/**
+		 * Add any required bone to the mask.
+		 * For example when we have finger bones that are not included in the bone inputs to the ML model, we still have to add those to the mask of the closest bone that is included (likely the hand or lower arm).
+		 * With closest we mean the closest up the hierarchy chain (traversing to the parent and the parent's parent, etc). We call those virtual parents in this context.
+		 * @param RefSkel The reference skeleton used for indexing and to get the hierarchy information from.
+		 * @param SkeletonBoneIndex The bone index to fill the mask for.
+		 * @param VirtualParentTable The table of virtual parents. This lets us know that we also have to add say the finger bone vertices to the mask of the hand in our example described above.
+		 * @param OutBonesAdded The list of bones we added to the mask. This list can contain items already. Bones that are already in the list won't be added to the mask again.
+		 */
+		void AddRequiredBones(const FReferenceSkeleton& RefSkel, int32 SkeletonBoneIndex, const TArray<int32>& VirtualParentTable, TArray<int32>& OutBonesAdded);
+
+		/**
+		 * Find the virtual parent for a given bone.
+		 * For more information see the BuildVirtualParentTable method.
+		 * @param RefSkel The reference skeleton.
+		 * @param BoneIndex The index of the bone in the reference skeleton to get the virtual parent for.
+		 * @param IncludedBoneNames The list of bones that the ML Deformer model uses as inputs.
+		 * @return The bone index that is the virtual parent for the specified bone. This can point to itself in case it is in the included bone names list.
+		 * @see BuildVirtualParentTable.
+		 */
+		int32 FindVirtualParentIndex(const FReferenceSkeleton& RefSkel, int32 BoneIndex, const TArray<FName>& IncludedBoneNames) const;
+
+		/**
+		 * Build a mapping table that tells us the virtual parent for each bone in the ref skeleton.
+		 * For example if we have finger bones that are not in the ML Deformer inputs, we try to find the closest parent bone in the hierarchy.
+		 * So for the fingers this might be the hand bone. The vertices touched by the fingers will then be added to the mask of the hand bone.
+		 * The table works like this: "Table[FingerBoneIndex] == HandBoneIndex". When the bone is inside the bone input list of the ML Deformer, it will
+		 * just return itself, and not its parent bone. Basically it just tells us to which bone's mask to write to for this specific bone.
+		 * @param RefSkel The reference skeleton.
+		 * @param IncludedBoneNames The names of the bones that are inputs to the ML model.
+		 * @return Returns a mapping array with the length of number of bones in the reference skeleton. Each array element contains a bone index inside the reference skeleton as well.
+		 */
+		TArray<int32> BuildVirtualParentTable(const FReferenceSkeleton& RefSkel, const TArray<FName>& IncludedBoneNames) const;
+
 		// Helpers.
 		UMLDeformerMorphModel* GetMorphModel() const;
 		UMLDeformerMorphModelVizSettings* GetMorphModelVizSettings() const;
@@ -54,13 +144,143 @@ namespace UE::MLDeformer
 		 */
 		void ClampMorphTargetNumber();
 
+		/**
+		 * Update the morph target error values. This assigns an error value for each morph target.
+		 * Internally this calls SetMorphTargetsErrorOrder to set the actual order and errors.
+		 * @param MorphTargets The morph targets to calculate the errors and order for.
+		 */
 		void UpdateMorphErrorValues(TArrayView<UMorphTarget*> MorphTargets);
 
+		/**
+		 * Calculate the normals for a given morph target.
+		 * This method is deprecated, please use the one that takes more parameters instead.
+		 * @param LOD The LOD level.
+		 * @param SkelMesh The skeletal mesh to get the mesh data from.
+		 * @param Deltas The per vertex deltas. The number of elements in this array is NumMorphTargets * NumImportedVertices.
+		 * @param BaseVertexPositions The positions of the base/neutral mesh, basically the unskinned vertices.
+		 * @param BaseNormals The normals of the base mesh.
+		 * @param OutDeltaNormals The array that we will write the generated normals to. This will automatically be resized by this method.
+		 */
+		UE_DEPRECATED(5.3, "Please use the CalcMorphTargetNormals method that takes more parameters.")
+		virtual void CalcMorphTargetNormals(int32 LOD, USkeletalMesh* SkelMesh, int32 MorphTargetIndex, TArrayView<const FVector3f> Deltas, TArrayView<const FVector3f> BaseVertexPositions, TArrayView<FVector3f> BaseNormals, TArray<FVector3f>& OutDeltaNormals);
+
+		/**
+		 * Calculate the delta normals for a given morph target.
+		 * @param LOD The LOD level.
+		 * @param SkelMesh The skeletal mesh to get the mesh data from.
+		 * @param MorphTargetIndex The morph target to calculate the normals for.
+		 * @param Deltas The per vertex deltas. The number of elements in this array is NumMorphTargets * NumImportedVertices.
+		 * @param BaseVertexPositions The positions of the base/neutral mesh, basically the unskinned vertices.
+		 * @param BaseNormals The normals of the neutral base mesh.
+		 * @param ImportedVertexToRenderVertexMapping The array of size 'NumImportedVerts' (NumBaseMeshVerts) that maps to a render vertex.
+		 * @param ColorBuffer The color buffer to use when calculating the global mask.
+		 * @param MaskChannel The global mask channel, for example which color channel to use from the color buffer.
+		 * @param bInvertGlobalMask Set to true when you want the global mask to be inverted.
+		 * @param OutDeltaNormals The array that we will write the generated normals to. This will automatically be resized by this method.
+		 */
+		virtual void CalcMorphTargetNormals(
+			int32 LOD,
+			const USkeletalMesh* SkelMesh,
+			int32 MorphTargetIndex,
+			const TArrayView<const FVector3f> Deltas,
+			const TArrayView<const FVector3f> BaseVertexPositions,
+			const TArrayView<const FVector3f> BaseNormals,
+			const TArrayView<const int32> ImportedVertexToRenderVertexMapping,
+			const FColorVertexBuffer& ColorBuffer,
+			EMLDeformerMaskChannel MaskChannel,
+			bool bInvertGlobalMaskChannel,
+			TArray<FVector3f>& OutDeltaNormals);
+
+		/**
+		 * Debug draw specific morph targets using lines and points.
+		 * This can show the user what deltas are included in which morph target.
+		 * @param PDI A pointer to the draw interface.
+		 * @param MorphDeltas A buffer of deltas for ALL morph targets. The size of the buffer must be a multiple of Model->GetBaseNumVerts().
+		 *        So the layout of this buffer is [Morph0_Deltas][Morph1_Deltas][Morph2_Deltas] etc.
+		 * @param DeltaThreshold Deltas with a length  larger or equal to the given threshold value will be colored differently than the ones smaller than this threshold.
+		 * @param MorphTargetIndex The morph target number to visualize.
+		 * @param DrawOffset An offset to perform the debug draw at.
+		 */
+		void DebugDrawMorphTarget(FPrimitiveDrawInterface* PDI, const TArray<FVector3f>& MorphDeltas, float DeltaThreshold, int32 MorphTargetIndex, const FVector& DrawOffset);
+
+		/** Zero all deltas with a length equal to, or smaller than the threshold value. */
+		void ZeroDeltasByLengthThreshold(TArray<FVector3f>& Deltas, float Threshold);
+
+		/**
+		 * Generate engine morph targets from a set of deltas.
+		 * @param OutMorphTargets The output array with generated morph targets. This array will be reset, and then filled with generated morph targets.
+		 * @param Deltas The per vertex deltas for all morph targets, as one big buffer. Each morph target has 'GetNumBaseMeshVerts()' number of deltas.
+		 * @param NamePrefix The morph target name prefix. If set to "MorphTarget_" the names will be "MorphTarget_000", "MorphTarget_001", "MorphTarget_002", etc.
+		 * @param LOD The LOD index to generate the morphs for.
+		 * @param DeltaThreshold Only include deltas with a length larger than this threshold in the morph targets.
+		 * @param bIncludeNormals Include normals inside the morph targets? This can be an alternative to recalculating normals at the end, although setting this to true and not 
+		 *        recomputing normals can lead to lower quality results, in trade for faster performance.
+		 * @param MaskChannel The weight mask mode, which specifies what channel to get the weight data from. Such channel allows the user to define what areas the deformer should for example not be active in.
+		 * @param bInvertMaskChannel Specifies whether the weight mask should be inverted or not.
+		 * @param MaskBuffer An optional mask buffer that contains 'Model->GetNumBaseMeshVerts() * (MorphModel->GetNumMorphTargets() - 1)' number of floats. Deltas will be multiplied by this value. 
+		 *        When the mask buffer is an empty array, it will be ignored.
+		 */
+		void CreateMorphTargets(
+			TArray<UMorphTarget*>& OutMorphTargets,
+			const TArray<FVector3f>& Deltas,
+			const FString& NamePrefix = TEXT("MorphTarget_"),
+			int32 LOD = 0,
+			float DeltaThreshold = 0.01f,
+			bool bIncludeNormals=false,
+			const EMLDeformerMaskChannel MaskChannel = EMLDeformerMaskChannel::Disabled,
+			bool bInvertMaskChannel = false);
+
+		/** 
+		 * Compress morph targets into GPU based morph buffers.
+		 * @param OutMorphBuffers The output compressed GPU based morph buffers. If this buffer is already initialized it will be released first.
+		 * @param MorphTargets The morph targets to compress into GPU friendly buffers.
+		 * @param LOD The LOD index to generate the morphs for.
+		 * @param MorphErrorTolerance The error tolerance for the delta compression, in cm. Higher values compress better but can result in artifacts.
+		 */
+		void CompressMorphTargets(FMorphTargetVertexInfoBuffers& OutMorphBuffers, const TArray<UMorphTarget*>& MorphTargets, int32 LOD = 0, float MorphErrorTolerance = 0.01f);
+
+		/**
+		 * Calculate the normals for each vertex, given the triangle data and positions.
+		 * It computes this by summing up the face normals for each vertex using that face, and normalizing them at the end.
+		 * @param VertexPositions The buffer with vertex positions. This is the size of the number of imported vertices.
+		 * @param IndexArray The index buffer, which contains NumTriangles * 3 number of integers.
+		 * @param VertexMap For each render vertex, an imported vertex number. For example, for a cube these indices go from 0..7.
+		 * @param OutNormals The array that will contain the normals. This will automatically be resized internally by this method.
+		 */
+		void CalcVertexNormals(TArrayView<const FVector3f> VertexPositions, TArrayView<const uint32> IndexArray, TArrayView<const int32> VertexMap, TArray<FVector3f>& OutNormals) const;
+
+		/**
+		 * Process and filter a vertex delta position and normal and apply scaling to it.
+		 * The scaling of these deltas is based on a given morph mask weight, a global mask weight, and a delta zero threshold, which filters out small deltas entirely.
+		 * @param OutScaledDelta The scaled version of the delta position.
+		 * @param OutScaledDeltaNormal The scaled version of the delta normal.
+		 * @param RawDelta The unscaled raw delta position.
+		 * @param RawDeltaNormal The unscaled raw delta normal.
+		 * @param DeltaThreshold Any scaled vertex delta position vector with a length smaller or equal to this value, will cause this method to return false. This basically can be used to filter out small deltas.
+		 * @param MorphMaskWeight The weight value of the morph target mask for this vertex. This can be values larger than 1 as well, but always larger or equal to 0.
+		 * @param GlobalMaskWeight The global weight value for this vertex. This global mask / weight is used to disable the deformer in specific regions on the mesh, for example where the head or neck connects with the body.
+		 * @return Returns true if this vertex should be included in the morph target, or false if it has been filtered out based on the provided delta threshold.
+		 */
+		bool ProcessVertexDelta(FVector3f& OutScaledDelta, FVector3f& OutScaledDeltaNormal, const FVector3f RawDelta, const FVector3f RawDeltaNormal, float DeltaThreshold, float MorphMaskWeight, float GlobalMaskWeight) const;
+
+		/**
+		 * Calculate the global mask weight for a specific render vertex, using a color buffer and specified mask channel settings.
+		 * @param RenderVertexIndex The render vertex index to calculate the global deformer mask weight for.
+		 * @param ColorBuffer The color buffer which we can grab color values from, which are used to calculate the mask weight.
+		 * @param MaskChannel The channel to grab the weight value from. This specifieds whether the mask is enabled and if so, from what channel (r, g, b, a) we grab the weight value.
+		 * @param bInvertMaskChannel Set to true when we should invert the weight value.
+		 * @return Returns the mask weight. A value of 1.0 is returned when the mask channel is set to disable the mask. The return value is always between 0 and 1.
+		 */
+		float CalcGlobalMaskWeight(int32 RenderVertexIndex, const FColorVertexBuffer& ColorBuffer, EMLDeformerMaskChannel MaskChannel, bool bInvertMaskChannel) const;
+
 	protected:
 		/**
 		 * The entire set of morph target deltas, 3 per vertex, for each morph target, as one flattened buffer.
 		 * So the size of this buffer is: (NumVertsPerMorphTarget * NumMorphTargets).
 		 */
 		TArray<FVector3f> MorphTargetDeltasBackup;
+
+		/** The backup of the input item mask buffer. This contains all the bone and bone group masks. */
+		TArray<float> InputItemMaskBufferBackup;
 	};
 }	// namespace UE::MLDeformer
diff --git a/Engine/Plugins/Animation/MLDeformer/MLDeformerFramework/Source/MLDeformerFrameworkEditor/Public/SMLDeformerInputBonesWidget.h b/Engine/Plugins/Animation/MLDeformer/MLDeformerFramework/Source/MLDeformerFrameworkEditor/Public/SMLDeformerInputBonesWidget.h
new file mode 100644
--- /dev/null
+++ b/Engine/Plugins/Animation/MLDeformer/MLDeformerFramework/Source/MLDeformerFrameworkEditor/Public/SMLDeformerInputBonesWidget.h
@@ -0,0 +1,131 @@
+// Copyright Epic Games, Inc. All Rights Reserved.
+
+#pragma once
+
+#include "CoreTypes.h"
+#include "Templates/SharedPointer.h"
+#include "Widgets/SCompoundWidget.h"
+#include "Widgets/Views/STreeView.h"
+#include "Framework/Commands/Commands.h"
+#include "Framework/MultiBox/MultiBoxBuilder.h"
+#include "UObject/NameTypes.h"
+#include "Styling/SlateColor.h"
+
+struct FReferenceSkeleton;
+
+namespace UE::MLDeformer
+{
+	class SMLDeformerInputBoneTreeWidget;
+	class SMLDeformerInputBonesWidget;
+	class SMLDeformerInputWidget;
+	class FMLDeformerEditorModel;
+
+	class MLDEFORMERFRAMEWORKEDITOR_API FMLDeformerInputBonesWidgetCommands
+		: public TCommands<FMLDeformerInputBonesWidgetCommands>
+	{
+	public:
+		FMLDeformerInputBonesWidgetCommands();
+		virtual void RegisterCommands() override;
+
+	public:
+		TSharedPtr<FUICommandInfo> AddInputBones;
+		TSharedPtr<FUICommandInfo> DeleteInputBones;
+		TSharedPtr<FUICommandInfo> ClearInputBones;
+		TSharedPtr<FUICommandInfo> AddAnimatedBones;
+	};
+
+
+	class MLDEFORMERFRAMEWORKEDITOR_API FMLDeformerInputBoneTreeElement
+		: public TSharedFromThis<FMLDeformerInputBoneTreeElement>
+	{
+	public:
+		TSharedRef<ITableRow> MakeTreeRowWidget(const TSharedRef<STableViewBase>& InOwnerTable, TSharedRef<FMLDeformerInputBoneTreeElement> InTreeElement, TSharedPtr<SMLDeformerInputBoneTreeWidget> InTreeWidget);
+
+	public:
+		FName Name;
+		TArray<TSharedPtr<FMLDeformerInputBoneTreeElement>> Children;
+		FSlateColor TextColor;
+	};
+
+
+	class MLDEFORMERFRAMEWORKEDITOR_API SMLDeformerInputBoneTreeRowWidget 
+		: public STableRow<TSharedPtr<FMLDeformerInputBoneTreeElement>>
+	{
+	public:
+		void Construct(const FArguments& InArgs, const TSharedRef<STableViewBase>& OwnerTable, TSharedRef<FMLDeformerInputBoneTreeElement> InTreeElement, TSharedPtr<SMLDeformerInputBoneTreeWidget> InTreeView);
+
+	private:
+		TWeakPtr<FMLDeformerInputBoneTreeElement> WeakTreeElement;
+		FText GetName() const;
+
+		friend class SMLDeformerInputBoneTreeWidget; 
+	};
+
+
+	class MLDEFORMERFRAMEWORKEDITOR_API SMLDeformerInputBoneTreeWidget
+		: public STreeView<TSharedPtr<FMLDeformerInputBoneTreeElement>>
+	{
+		SLATE_BEGIN_ARGS(SMLDeformerInputBoneTreeWidget) {}
+		SLATE_ARGUMENT(TSharedPtr<SMLDeformerInputBonesWidget>, InputBonesWidget)
+		SLATE_END_ARGS()
+
+	public:
+		void Construct(const FArguments& InArgs);
+		const TArray<TSharedPtr<FMLDeformerInputBoneTreeElement>>& GetRootElements() const	{ return RootElements; }
+		void AddElement(TSharedPtr<FMLDeformerInputBoneTreeElement> Element, TSharedPtr<FMLDeformerInputBoneTreeElement> ParentElement);
+		void RefreshElements(const TArray<FName>& BoneNames, const FReferenceSkeleton* RefSkeleton);
+		TArray<FName> ExtractAllElementNames() const;
+
+	private:
+		TSharedRef<ITableRow> MakeTableRowWidget(TSharedPtr<FMLDeformerInputBoneTreeElement> InItem, const TSharedRef<STableViewBase>& OwnerTable);
+		void HandleGetChildrenForTree(TSharedPtr<FMLDeformerInputBoneTreeElement> InItem, TArray<TSharedPtr<FMLDeformerInputBoneTreeElement>>& OutChildren);
+		FReply OnKeyDown(const FGeometry& InGeometry, const FKeyEvent& InKeyEvent) override;
+		TSharedPtr<SWidget> OnContextMenuOpening() const;
+		void RecursiveSortElements(TSharedPtr<FMLDeformerInputBoneTreeElement> Element);
+		void RecursiveAddNames(const FMLDeformerInputBoneTreeElement& Element, TArray<FName>& OutNames) const;
+		void OnSelectionChanged(TSharedPtr<FMLDeformerInputBoneTreeElement> Selection, ESelectInfo::Type SelectInfo);
+
+	private:
+		TArray<TSharedPtr<FMLDeformerInputBoneTreeElement>> RootElements;
+		TSharedPtr<SMLDeformerInputBonesWidget> InputBonesWidget;
+	};
+
+
+	class MLDEFORMERFRAMEWORKEDITOR_API SMLDeformerInputBonesWidget
+		: public SCompoundWidget
+	{
+		friend class SMLDeformerInputBoneTreeWidget;
+
+		SLATE_BEGIN_ARGS(SMLDeformerInputBonesWidget) {}
+		SLATE_ARGUMENT(FMLDeformerEditorModel*, EditorModel)
+		SLATE_ARGUMENT(TSharedPtr<SMLDeformerInputWidget>, InputWidget)
+		SLATE_END_ARGS()
+
+	public:
+		void Construct(const FArguments& InArgs);
+		void Refresh();
+		void BindCommands(TSharedPtr<FUICommandList> CommandList);
+		TSharedPtr<SMLDeformerInputBoneTreeWidget> GetTreeWidget() const;
+		TSharedPtr<SMLDeformerInputWidget> GetInputWidget() const;
+		FText GetSectionTitle() const;
+
+	private:
+		void OnFilterTextChanged(const FText& InFilterText);
+		void RefreshTree(bool bBroadCastPropertyChanged=true);
+		bool BroadcastModelPropertyChanged(const FName PropertyName);
+
+		void OnAddInputBones();
+		void OnDeleteInputBones();
+		void OnClearInputBones();
+		void OnAddAnimatedBones();
+
+	private:
+		TSharedPtr<SMLDeformerInputBoneTreeWidget> TreeWidget;
+		TSharedPtr<SMLDeformerInputWidget> InputWidget;
+		FMLDeformerEditorModel* EditorModel = nullptr;
+		FString FilterText;
+		FText SectionTitle;
+	};
+
+}	// namespace UE::MLDeformer
+
diff --git a/Engine/Plugins/Animation/MLDeformer/MLDeformerFramework/Source/MLDeformerFrameworkEditor/Public/SMLDeformerInputCurvesWidget.h b/Engine/Plugins/Animation/MLDeformer/MLDeformerFramework/Source/MLDeformerFrameworkEditor/Public/SMLDeformerInputCurvesWidget.h
new file mode 100644
--- /dev/null
+++ b/Engine/Plugins/Animation/MLDeformer/MLDeformerFramework/Source/MLDeformerFrameworkEditor/Public/SMLDeformerInputCurvesWidget.h
@@ -0,0 +1,127 @@
+// Copyright Epic Games, Inc. All Rights Reserved.
+
+#pragma once
+
+#include "CoreTypes.h"
+#include "Templates/SharedPointer.h"
+#include "Widgets/SCompoundWidget.h"
+#include "Widgets/Views/SListView.h"
+#include "Framework/Commands/Commands.h"
+#include "Framework/MultiBox/MultiBoxBuilder.h"
+#include "UObject/NameTypes.h"
+#include "Styling/SlateColor.h"
+
+class USkeleton;
+
+namespace UE::MLDeformer
+{
+	class SMLDeformerInputCurveListWidget;
+	class SMLDeformerInputCurvesWidget;
+	class SMLDeformerInputWidget;
+	class FMLDeformerEditorModel;
+
+	class MLDEFORMERFRAMEWORKEDITOR_API FMLDeformerInputCurvesWidgetCommands
+		: public TCommands<FMLDeformerInputCurvesWidgetCommands>
+	{
+	public:
+		FMLDeformerInputCurvesWidgetCommands();
+		virtual void RegisterCommands() override;
+
+	public:
+		TSharedPtr<FUICommandInfo> AddInputCurves;
+		TSharedPtr<FUICommandInfo> DeleteInputCurves;
+		TSharedPtr<FUICommandInfo> ClearInputCurves;
+		TSharedPtr<FUICommandInfo> AddAnimatedCurves;
+	};
+
+
+	class MLDEFORMERFRAMEWORKEDITOR_API FMLDeformerInputCurveListElement
+		: public TSharedFromThis<FMLDeformerInputCurveListElement>
+	{
+	public:
+		TSharedRef<ITableRow> MakeRowWidget(const TSharedRef<STableViewBase>& InOwnerTable, TSharedRef<FMLDeformerInputCurveListElement> InTreeElement, TSharedPtr<SMLDeformerInputCurveListWidget> InTreeWidget);
+
+	public:
+		FName Name;
+		FSlateColor TextColor;
+	};
+
+
+	class MLDEFORMERFRAMEWORKEDITOR_API SMLDeformerInputCurveListRowWidget 
+		: public STableRow<TSharedPtr<FMLDeformerInputCurveListElement>>
+	{
+	public:
+		void Construct(const FArguments& InArgs, const TSharedRef<STableViewBase>& OwnerTable, TSharedRef<FMLDeformerInputCurveListElement> InElement, TSharedPtr<SMLDeformerInputCurveListWidget> InListView);
+
+	private:
+		TWeakPtr<FMLDeformerInputCurveListElement> WeakElement;
+		FText GetName() const;
+
+		friend class SMLDeformerInputCurveListWidget; 
+	};
+
+
+	class MLDEFORMERFRAMEWORKEDITOR_API SMLDeformerInputCurveListWidget
+		: public SListView<TSharedPtr<FMLDeformerInputCurveListElement>>
+	{
+		SLATE_BEGIN_ARGS(SMLDeformerInputCurveListWidget) {}
+		SLATE_ARGUMENT(TSharedPtr<SMLDeformerInputCurvesWidget>, InputCurvesWidget)
+		SLATE_END_ARGS()
+
+	public:
+		void Construct(const FArguments& InArgs);
+		const TArray<TSharedPtr<FMLDeformerInputCurveListElement>>& GetElements() const	{ return Elements; }
+		void AddElement(TSharedPtr<FMLDeformerInputCurveListElement> Element);
+		void RefreshElements(const TArray<FName>& CurveNames, const USkeleton* Skeleton);
+
+	private:
+		TSharedRef<ITableRow> MakeTableRowWidget(TSharedPtr<FMLDeformerInputCurveListElement> InItem, const TSharedRef<STableViewBase>& OwnerTable);
+		FReply OnKeyDown(const FGeometry& InGeometry, const FKeyEvent& InKeyEvent) override;
+		TSharedPtr<SWidget> OnContextMenuOpening() const;
+		void SortElements();
+		void OnSelectionChanged(TSharedPtr<FMLDeformerInputCurveListElement> Selection, ESelectInfo::Type SelectInfo);
+
+	private:
+		TArray<TSharedPtr<FMLDeformerInputCurveListElement>> Elements;
+		TSharedPtr<SMLDeformerInputCurvesWidget> InputCurvesWidget;
+	};
+
+
+	class MLDEFORMERFRAMEWORKEDITOR_API SMLDeformerInputCurvesWidget
+		: public SCompoundWidget
+	{
+		friend class SMLDeformerInputCurveListWidget;
+
+		SLATE_BEGIN_ARGS(SMLDeformerInputCurvesWidget) {}
+		SLATE_ARGUMENT(FMLDeformerEditorModel*, EditorModel)
+		SLATE_ARGUMENT(TSharedPtr<SMLDeformerInputWidget>, InputWidget)
+		SLATE_END_ARGS()
+
+	public:
+		void Construct(const FArguments& InArgs);
+		void Refresh();
+		void BindCommands(TSharedPtr<FUICommandList> CommandList);
+		FText GetSectionTitle() const;
+		TSharedPtr<SMLDeformerInputCurveListWidget> GetListWidget() const;
+		TSharedPtr<SMLDeformerInputWidget> GetInputWidget() const;
+
+	private:
+		void OnFilterTextChanged(const FText& InFilterText);
+		void RefreshList(bool bBroadCastPropertyChanged=true);
+		bool BroadcastModelPropertyChanged(const FName PropertyName);
+
+		void OnAddInputCurves();
+		void OnDeleteInputCurves();
+		void OnClearInputCurves();
+		void OnAddAnimatedCurves();
+
+	private:
+		TSharedPtr<SMLDeformerInputWidget> InputWidget;
+		TSharedPtr<SMLDeformerInputCurveListWidget> ListWidget;
+		FMLDeformerEditorModel* EditorModel = nullptr;
+		FString FilterText;
+		FText SectionTitle;
+	};
+
+}	// namespace UE::MLDeformer
+
diff --git a/Engine/Plugins/Animation/MLDeformer/NearestNeighborModel/Source/NearestNeighborModel/Private/NearestNeighborModel.cpp b/Engine/Plugins/Animation/MLDeformer/NearestNeighborModel/Source/NearestNeighborModel/Private/NearestNeighborModel.cpp
--- a/Engine/Plugins/Animation/MLDeformer/NearestNeighborModel/Source/NearestNeighborModel/Private/NearestNeighborModel.cpp
+++ b/Engine/Plugins/Animation/MLDeformer/NearestNeighborModel/Source/NearestNeighborModel/Private/NearestNeighborModel.cpp
@@ -458,6 +458,7 @@ bool UNearestNeighborModel::CheckPCAData(int32 PartId) const
 void UNearestNeighborModel::SetOptimizedNetwork(UNearestNeighborOptimizedNetwork* InOptimizedNetwork)
 {
 	OptimizedNetwork = InOptimizedNetwork;
+	GetReinitModelInstanceDelegate().Broadcast();
 }
 
 bool UNearestNeighborModel::DoesUseOptimizedNetwork() const
diff --git a/Engine/Plugins/Animation/MLDeformer/NearestNeighborModel/Source/NearestNeighborModel/Public/NearestNeighborModelInputInfo.h b/Engine/Plugins/Animation/MLDeformer/NearestNeighborModel/Source/NearestNeighborModel/Public/NearestNeighborModelInputInfo.h
--- a/Engine/Plugins/Animation/MLDeformer/NearestNeighborModel/Source/NearestNeighborModel/Public/NearestNeighborModelInputInfo.h
+++ b/Engine/Plugins/Animation/MLDeformer/NearestNeighborModel/Source/NearestNeighborModel/Public/NearestNeighborModelInputInfo.h
@@ -2,7 +2,7 @@
 
 #pragma once
 
-#include "MLDeformerInputInfo.h"
+#include "MLDeformerMorphModelInputInfo.h"
 #include "NearestNeighborModelInputInfo.generated.h"
 
 class UNearestNeighborModelInstance;
@@ -14,7 +14,7 @@ namespace UE::NearestNeighborModel
 
 UCLASS()
 class NEARESTNEIGHBORMODEL_API UNearestNeighborModelInputInfo
-    : public UMLDeformerInputInfo
+    : public UMLDeformerMorphModelInputInfo
 {
     GENERATED_BODY()
 
diff --git a/Engine/Plugins/Animation/MLDeformer/NearestNeighborModel/Source/NearestNeighborModelEditor/Private/NearestNeighborEditorModel.cpp b/Engine/Plugins/Animation/MLDeformer/NearestNeighborModel/Source/NearestNeighborModelEditor/Private/NearestNeighborEditorModel.cpp
--- a/Engine/Plugins/Animation/MLDeformer/NearestNeighborModel/Source/NearestNeighborModelEditor/Private/NearestNeighborEditorModel.cpp
+++ b/Engine/Plugins/Animation/MLDeformer/NearestNeighborModel/Source/NearestNeighborModelEditor/Private/NearestNeighborEditorModel.cpp
@@ -750,7 +750,7 @@ namespace UE::NearestNeighborModel
 
 		const int32 LOD = 0;
 		TArray<UMorphTarget*> MorphTargets;
-		CreateEngineMorphTargets(
+		CreateMorphTargets(
 			MorphTargets, 
 			Deltas, 
 			FString("NNMorphTarget_"),
@@ -762,7 +762,7 @@ namespace UE::NearestNeighborModel
 
 		check(NearestNeighborModel->GetMorphTargetSet().IsValid());
 		FMorphTargetVertexInfoBuffers& MorphBuffers = NearestNeighborModel->GetMorphTargetSet()->MorphBuffers;
-		CompressEngineMorphTargets(MorphBuffers, MorphTargets, LOD, NearestNeighborModel->GetMorphCompressionLevel());
+		CompressMorphTargets(MorphBuffers, MorphTargets, LOD, NearestNeighborModel->GetMorphCompressionLevel());
 
 		if (MorphBuffers.GetNumBatches() <= 0)
 		{
diff --git a/Engine/Plugins/Animation/MLDeformer/NeuralMorphModel/Source/NeuralMorphModel/Private/NeuralMorphInputInfo.cpp b/Engine/Plugins/Animation/MLDeformer/NeuralMorphModel/Source/NeuralMorphModel/Private/NeuralMorphInputInfo.cpp
--- a/Engine/Plugins/Animation/MLDeformer/NeuralMorphModel/Source/NeuralMorphModel/Private/NeuralMorphInputInfo.cpp
+++ b/Engine/Plugins/Animation/MLDeformer/NeuralMorphModel/Source/NeuralMorphModel/Private/NeuralMorphInputInfo.cpp
@@ -17,11 +17,21 @@ TArray<FNeuralMorphBoneGroup>& UNeuralMorphInputInfo::GetBoneGroups()
 	return BoneGroups;
 }
 
+const TArray<FNeuralMorphBoneGroup>& UNeuralMorphInputInfo::GetBoneGroups() const
+{
+	return BoneGroups;
+}
+
 TArray<FNeuralMorphCurveGroup>& UNeuralMorphInputInfo::GetCurveGroups()
 {
 	return CurveGroups;
 }
 
+const TArray<FNeuralMorphCurveGroup>& UNeuralMorphInputInfo::GetCurveGroups() const
+{
+	return CurveGroups;
+}
+
 int32 UNeuralMorphInputInfo::CalcMaxNumGroupItems() const
 {
 	int32 MaxGroupSize = 0;
diff --git a/Engine/Plugins/Animation/MLDeformer/NeuralMorphModel/Source/NeuralMorphModel/Private/NeuralMorphModel.cpp b/Engine/Plugins/Animation/MLDeformer/NeuralMorphModel/Source/NeuralMorphModel/Private/NeuralMorphModel.cpp
--- a/Engine/Plugins/Animation/MLDeformer/NeuralMorphModel/Source/NeuralMorphModel/Private/NeuralMorphModel.cpp
+++ b/Engine/Plugins/Animation/MLDeformer/NeuralMorphModel/Source/NeuralMorphModel/Private/NeuralMorphModel.cpp
@@ -47,19 +47,24 @@ UMLDeformerInputInfo* UNeuralMorphModel::CreateInputInfo()
 	return NewObject<UNeuralMorphInputInfo>(this);
 }
 
-void UNeuralMorphModel::SetNeuralMorphNetwork(UNeuralMorphNetwork* Net)
-{ 
-	NeuralMorphNetwork = Net;
-}
-
 void UNeuralMorphModel::Serialize(FArchive& Archive)
 {
+	if (Archive.IsSaving())
+	{
+		UpdateMissingGroupNames();
+	}
+
 	if (Archive.IsSaving() && Archive.IsCooking())
 	{
+		// We haven't got a trained network, let's log a message about this, as it might be overlooked.
 		if (NeuralMorphNetwork == nullptr)
 		{
 			UE_LOG(LogNeuralMorphModel, Display, TEXT("Neural Morph Model in MLD asset '%s' still needs to be trained."), *GetDeformerAsset()->GetName());
 		}
+
+		// Strip the mask data in the cooked asset.
+		BoneMaskInfos.Empty();
+		BoneGroupMaskInfos.Empty();
 	}
 
 	// Convert the UMLDeformerInputInfo object into a UNeuralMorphInputInfo object for backward compatiblity.
@@ -77,4 +82,39 @@ void UNeuralMorphModel::Serialize(FArchive& Archive)
 	Super::Serialize(Archive);
 }
 
+void UNeuralMorphModel::UpdateMissingGroupNames()
+{
+	// Auto set names for the bone groups if they haven't been set yet.
+	for (int32 Index = 0; Index < BoneGroups.Num(); ++Index)
+	{
+		FNeuralMorphBoneGroup& BoneGroup = BoneGroups[Index];
+		if (!BoneGroup.GroupName.IsValid() || BoneGroup.GroupName.IsNone())
+		{
+			BoneGroup.GroupName = FName(FString::Format(TEXT("Bone Group #{0}"), {Index}));
+		}
+	}
+
+	// Auto set names for the curve groups if they haven't been set yet.
+	for (int32 Index = 0; Index < CurveGroups.Num(); ++Index)
+	{
+		FNeuralMorphCurveGroup& CurveGroup = CurveGroups[Index];
+		if (!CurveGroup.GroupName.IsValid() || CurveGroup.GroupName.IsNone())
+		{
+			CurveGroup.GroupName = FName(FString::Format(TEXT("Curve Group #{0}"), {Index}));
+		}
+	}
+}
+
+void UNeuralMorphModel::PostLoad()
+{
+	Super::PostLoad();
+	UpdateMissingGroupNames();
+}
+
+void UNeuralMorphModel::SetNeuralMorphNetwork(UNeuralMorphNetwork* Net)
+{ 
+	NeuralMorphNetwork = Net;
+	GetReinitModelInstanceDelegate().Broadcast();
+}
+
 #undef LOCTEXT_NAMESPACE
diff --git a/Engine/Plugins/Animation/MLDeformer/NeuralMorphModel/Source/NeuralMorphModel/Public/NeuralMorphInputInfo.h b/Engine/Plugins/Animation/MLDeformer/NeuralMorphModel/Source/NeuralMorphModel/Public/NeuralMorphInputInfo.h
--- a/Engine/Plugins/Animation/MLDeformer/NeuralMorphModel/Source/NeuralMorphModel/Public/NeuralMorphInputInfo.h
+++ b/Engine/Plugins/Animation/MLDeformer/NeuralMorphModel/Source/NeuralMorphModel/Public/NeuralMorphInputInfo.h
@@ -2,7 +2,7 @@
 #pragma once
 
 #include "CoreTypes.h"
-#include "MLDeformerInputInfo.h"
+#include "MLDeformerMorphModelInputInfo.h"
 #include "NeuralMorphTypes.h"
 #include "NeuralMorphInputInfo.generated.h"
 
@@ -11,14 +11,14 @@
  */
 UCLASS()
 class NEURALMORPHMODEL_API UNeuralMorphInputInfo
-	: public UMLDeformerInputInfo
+	: public UMLDeformerMorphModelInputInfo
 {
 	GENERATED_BODY()
 
 public:
 	// FMLDeformerInputInfo overrides.
 	virtual void Reset() override;
-	// ~END FMLDeformerInputInfooverrides.
+	// ~END FMLDeformerInputInfo overrides.
 
 	/**
 	 * Get the bone groups.
@@ -26,6 +26,7 @@ public:
 	 * @return A reference to the bone groups array.
 	 */
 	TArray<FNeuralMorphBoneGroup>& GetBoneGroups();
+	const TArray<FNeuralMorphBoneGroup>& GetBoneGroups() const;
 
 	/**
 	 * Get the curve groups.
@@ -33,6 +34,7 @@ public:
 	 * @return A reference to the curve groups array.
 	 */
 	TArray<FNeuralMorphCurveGroup>& GetCurveGroups();
+	const TArray<FNeuralMorphCurveGroup>& GetCurveGroups() const;
 
 	/**
 	 * Calculate the maximum number of group items.
diff --git a/Engine/Plugins/Animation/MLDeformer/NeuralMorphModel/Source/NeuralMorphModel/Public/NeuralMorphModel.h b/Engine/Plugins/Animation/MLDeformer/NeuralMorphModel/Source/NeuralMorphModel/Public/NeuralMorphModel.h
--- a/Engine/Plugins/Animation/MLDeformer/NeuralMorphModel/Source/NeuralMorphModel/Public/NeuralMorphModel.h
+++ b/Engine/Plugins/Animation/MLDeformer/NeuralMorphModel/Source/NeuralMorphModel/Public/NeuralMorphModel.h
@@ -5,6 +5,7 @@
 #include "MLDeformerMorphModel.h"
 #include "GeometryCache.h"
 #include "NeuralMorphTypes.h"
+#include "Containers/Map.h"
 #include "NeuralMorphModel.generated.h"
 
 class USkeletalMesh;
@@ -49,29 +50,34 @@ public:
 
 	// UObject overrides.
 	virtual void Serialize(FArchive& Archive) override;
+	virtual void PostLoad() override;
 	// ~END UObject overrides.
 
 	// UMLDeformerModel overrides.
-	virtual FString GetDisplayName() const override			{ return "Neural Morph Model"; }
+	virtual FString GetDisplayName() const override					{ return "Neural Morph Model"; }
 	virtual UMLDeformerModelInstance* CreateModelInstance(UMLDeformerComponent* Component) override;
-	virtual UMLDeformerInputInfo* CreateInputInfo();
+	virtual UMLDeformerInputInfo* CreateInputInfo() override;
 	// ~END UMLDeformerModel overrides.
 
-	void SetNeuralMorphNetwork(UNeuralMorphNetwork* Net);
-
 	const TArray<FNeuralMorphBoneGroup>& GetBoneGroups() const		{ return BoneGroups; }
 	const TArray<FNeuralMorphCurveGroup>& GetCurveGroups() const	{ return CurveGroups; }
-	UNeuralMorphNetwork* GetNeuralMorphNetwork() const		{ return NeuralMorphNetwork.Get(); }
-	ENeuralMorphMode GetModelMode() const					{ return Mode; }
-	int32 GetLocalNumHiddenLayers() const					{ return LocalNumHiddenLayers; }
-	int32 GetLocalNumNeuronsPerLayer() const				{ return LocalNumNeuronsPerLayer; }
-	int32 GetGlobalNumHiddenLayers() const					{ return GlobalNumHiddenLayers; }
-	int32 GetGlobalNumNeuronsPerLayer() const				{ return GlobalNumNeuronsPerLayer; }
-	int32 GetNumIterations() const							{ return NumIterations; }
-	int32 GetBatchSize() const								{ return BatchSize; }
-	float GetLearningRate() const							{ return LearningRate; }
-	float GetLearningRateDecay() const						{ return LearningRateDecay; }
-	float GetRegularizationFactor() const					{ return RegularizationFactor; }
+	UNeuralMorphNetwork* GetNeuralMorphNetwork() const				{ return NeuralMorphNetwork.Get(); }
+	ENeuralMorphMode GetModelMode() const							{ return Mode; }
+	int32 GetLocalNumMorphsPerBone() const							{ return LocalNumMorphTargetsPerBone; }
+	int32 GetLocalNumHiddenLayers() const							{ return LocalNumHiddenLayers; }
+	int32 GetLocalNumNeuronsPerLayer() const						{ return LocalNumNeuronsPerLayer; }
+	int32 GetGlobalNumMorphs() const								{ return GlobalNumMorphTargets; }
+	int32 GetGlobalNumHiddenLayers() const							{ return GlobalNumHiddenLayers; }
+	int32 GetGlobalNumNeuronsPerLayer() const						{ return GlobalNumNeuronsPerLayer; }
+	int32 GetNumIterations() const									{ return NumIterations; }
+	int32 GetBatchSize() const										{ return BatchSize; }
+	float GetLearningRate() const									{ return LearningRate; }
+	float GetLearningRateDecay() const								{ return LearningRateDecay; }
+	float GetRegularizationFactor() const							{ return RegularizationFactor; }
+	bool IsBoneMaskingEnabled() const								{ return bEnableBoneMasks; }
+
+	void UpdateMissingGroupNames();
+	void SetNeuralMorphNetwork(UNeuralMorphNetwork* Net);
 
 public:
 	/**
@@ -79,7 +85,7 @@ public:
 	 * This can be used in case multiple bones are correlated to each other and work together to produce given shapes.
 	 * Groups are only used when the model is in Local mode.
 	 */
-	UPROPERTY(EditAnywhere, Category = "Inputs and Output")
+	UPROPERTY()
 	TArray<FNeuralMorphBoneGroup> BoneGroups;
 
 	/**
@@ -87,9 +93,27 @@ public:
 	 * This can be used in case multiple curves are correlated to each other and work together to produce given shapes.
 	 * Groups are only used when the model is in Local mode.
 	 */
-	UPROPERTY(EditAnywhere, Category = "Inputs and Output")
+	UPROPERTY()
 	TArray<FNeuralMorphCurveGroup> CurveGroups;
 
+	/**
+	 * Information needed to generate a mask for each bone.
+	 * Each mask info object contains a list of bones who's skinning influence regions should be included in the final mask for the specific bone.
+	 * This information (and bone masking in general) is used inside the ENeuralMorphMode::Local mode.
+	 * The FName map key represents the bone name.
+	 */
+	UPROPERTY()
+	TMap<FName, FNeuralMorphMaskInfo> BoneMaskInfos;
+
+	/**
+	 * Information needed to generate a mask for each bone group.
+	 * So the size of the array is NumGroupsInInputInfo.
+	 * Each mask info object contains a list of bones who's skinning influence regions should be included in the final mask for the specific bone.
+	 * This information (and bone masking in general) is used inside the ENeuralMorphMode::Local mode.
+	 */
+	UPROPERTY()
+	TMap<FName, FNeuralMorphMaskInfo> BoneGroupMaskInfos;
+
 	/**
 	 * The mode that the neural network will operate in. 
 	 * Local mode means there is one tiny network per bone, while global mode has one network for all bones together.
@@ -169,6 +193,15 @@ public:
 	UPROPERTY(EditAnywhere, BlueprintReadWrite, AdvancedDisplay, Category = "Training Settings", meta = (ClampMin = "0.0", ClampMax = "10.0"))
 	float RegularizationFactor = 1.0f;
 
+	/** 
+	 * Enable the use of per bone and bone group masks.
+	 * When enabled, an influence mask is generated per bone based on skinning info. This will enforce deformations localized to the area around the joint.
+	 * The benefit of enabling this can be reduced GPU memory footprint, faster GPU performance and more localized deformations.
+	 * If deformations do not happen near the joint then this enabling this setting can lead to those deformations possibly not being captured.
+	 */
+	UPROPERTY(EditAnywhere, BlueprintReadWrite, AdvancedDisplay, Category = "Training Settings", meta = (EditCondition = "Mode == ENeuralMorphMode::Local"))
+	bool bEnableBoneMasks = false;
+
 	/**
 	 * The neural morph model network.
 	 * This is a specialized network used for inference, for performance reasons.
diff --git a/Engine/Plugins/Animation/MLDeformer/NeuralMorphModel/Source/NeuralMorphModel/Public/NeuralMorphModelVizSettings.h b/Engine/Plugins/Animation/MLDeformer/NeuralMorphModel/Source/NeuralMorphModel/Public/NeuralMorphModelVizSettings.h
--- a/Engine/Plugins/Animation/MLDeformer/NeuralMorphModel/Source/NeuralMorphModel/Public/NeuralMorphModelVizSettings.h
+++ b/Engine/Plugins/Animation/MLDeformer/NeuralMorphModel/Source/NeuralMorphModel/Public/NeuralMorphModelVizSettings.h
@@ -4,6 +4,7 @@
 
 #include "CoreTypes.h"
 #include "MLDeformerMorphModelVizSettings.h"
+#include "NeuralMorphTypes.h"
 #include "NeuralMorphModelVizSettings.generated.h"
 
 /**
@@ -16,4 +17,15 @@ class NEURALMORPHMODEL_API UNeuralMorphModelVizSettings
 	: public UMLDeformerMorphModelVizSettings
 {
 	GENERATED_BODY()
+
+public:
+	/**
+	 * The visualization mode for the masks.
+	 * Each bone, curve, bone group or curve group has a specific mask area on the mesh.
+	 * This mask defines areas where generated morph targets can be active. They can be used to filter out deformations in undesired areas.
+	 * For example if you rotate the left arm, you don't want the right arm to deform. The mask for the left arm can be setup in a way that it only includes 
+	 * vertices around the area of the left arm to enforce this.
+	 */
+	UPROPERTY(EditAnywhere, Category = "Live Settings")
+	ENeuralMorphMaskVizMode MaskVizMode = ENeuralMorphMaskVizMode::WhenInFocus;
 };
diff --git a/Engine/Plugins/Animation/MLDeformer/NeuralMorphModel/Source/NeuralMorphModel/Public/NeuralMorphTypes.h b/Engine/Plugins/Animation/MLDeformer/NeuralMorphModel/Source/NeuralMorphModel/Public/NeuralMorphTypes.h
--- a/Engine/Plugins/Animation/MLDeformer/NeuralMorphModel/Source/NeuralMorphModel/Public/NeuralMorphTypes.h
+++ b/Engine/Plugins/Animation/MLDeformer/NeuralMorphModel/Source/NeuralMorphModel/Public/NeuralMorphTypes.h
@@ -33,6 +33,26 @@ enum class ENeuralMorphMode : uint8
 	Global
 };
 
+/**
+ * The visualization mode for the masks.
+ * Each bone, curve, bone group or curve group has a specific mask area on the mesh.
+ * This mask defines areas where generated morph targets can be active. They can be used to filter out deformations in undesired areas.
+ * For example if you rotate the left arm, you don't want the right arm to deform. The mask for the left arm can be setup in a way that it only includes 
+ * vertices around the area of the left arm to enforce this.
+ */
+UENUM()
+enum class ENeuralMorphMaskVizMode : uint8
+{
+	/** Do not display the masks in the viewport. */
+	Off,
+
+	/** Only show the masks inside the viewport when the inputs widget on the right side of the UI is in focus. Show the mask for the selected item (bone, curve, bone group, curve group). */
+	WhenInFocus,
+
+	/** Always show the selected mask. The mask selected is defined by what is selected in the input widget on the right side of the UI. */
+	Always
+};
+
 /**
  * A group of bones, which can generate morph targets together.
  * This is useful when there are specific correlations between different bones, that all effect the same region on the body.
@@ -43,6 +63,10 @@ struct NEURALMORPHMODEL_API FNeuralMorphBoneGroup
 	GENERATED_BODY()
 
 public:
+	/** The name of the group, also shown in the UI. */
+	UPROPERTY()
+	FName GroupName;
+
 	/** The list of bones that should form a group together. */
 	UPROPERTY(EditAnywhere, Category = "Bones")
 	TArray<FBoneReference> BoneNames;
@@ -58,7 +82,29 @@ struct NEURALMORPHMODEL_API FNeuralMorphCurveGroup
 	GENERATED_BODY()
 
 public:
+	/** The name of the group, also shown in the UI. */
+	UPROPERTY()
+	FName GroupName;
+
 	/** The list of curves that should form a group together. */
 	UPROPERTY(EditAnywhere, Category = "Curves")
 	TArray<FMLDeformerCurveReference> CurveNames;
 };
+
+
+/**
+ * Information needed to generate the mask for a specific bone.
+ * This includes a list of bone names. Each bone has a skinning influence mask.
+ * The final mask will be a merge of all the bone masks of bones listed inside this info struct.
+ * There will be an array of these structs, one for each bone.
+ */
+USTRUCT()
+struct NEURALMORPHMODEL_API FNeuralMorphMaskInfo
+{
+	GENERATED_BODY()
+
+public:
+	/** The list of bone names that should be included in the mask generation. */
+	UPROPERTY()
+	TArray<FName> BoneNames;
+};
diff --git a/Engine/Plugins/Animation/MLDeformer/NeuralMorphModel/Source/NeuralMorphModelEditor/NeuralMorphModelEditor.Build.cs b/Engine/Plugins/Animation/MLDeformer/NeuralMorphModel/Source/NeuralMorphModelEditor/NeuralMorphModelEditor.Build.cs
--- a/Engine/Plugins/Animation/MLDeformer/NeuralMorphModel/Source/NeuralMorphModelEditor/NeuralMorphModelEditor.Build.cs
+++ b/Engine/Plugins/Animation/MLDeformer/NeuralMorphModel/Source/NeuralMorphModelEditor/NeuralMorphModelEditor.Build.cs
@@ -18,6 +18,7 @@ namespace UnrealBuildTool.Rules
 			PrivateDependencyModuleNames.AddRange(
 				new string[]
 				{
+					"InputCore",
 					"UnrealEd",
 					"CoreUObject",
 					"Engine",
@@ -27,7 +28,7 @@ namespace UnrealBuildTool.Rules
 					"MLDeformerFrameworkEditor",
 					"NeuralMorphModel",
 					"PropertyEditor",
-					"ToolWidgets",
+					"ToolWidgets"
 				}
 			);
 
diff --git a/Engine/Plugins/Animation/MLDeformer/NeuralMorphModel/Source/NeuralMorphModelEditor/Private/NeuralMorphEditorModel.cpp b/Engine/Plugins/Animation/MLDeformer/NeuralMorphModel/Source/NeuralMorphModelEditor/Private/NeuralMorphEditorModel.cpp
--- a/Engine/Plugins/Animation/MLDeformer/NeuralMorphModel/Source/NeuralMorphModelEditor/Private/NeuralMorphEditorModel.cpp
+++ b/Engine/Plugins/Animation/MLDeformer/NeuralMorphModel/Source/NeuralMorphModelEditor/Private/NeuralMorphEditorModel.cpp
@@ -6,11 +6,25 @@
 #include "NeuralMorphNetwork.h"
 #include "NeuralMorphInputInfo.h"
 #include "NeuralMorphTrainingModel.h"
+#include "NeuralMorphModelVizSettings.h"
 #include "MLDeformerMorphModelVizSettings.h"
 #include "MLDeformerEditorToolkit.h"
+#include "MLDeformerEditorActor.h"
+#include "SMLDeformerInputWidget.h"
 #include "Engine/SkeletalMesh.h"
 #include "Animation/Skeleton.h"
+#include "Animation/DebugSkelMeshComponent.h"
 #include "BoneContainer.h"
+#include "BoneWeights.h"
+#include "Rendering/SkeletalMeshLODRenderData.h"
+#include "Rendering/SkeletalMeshRenderData.h"
+#include "Rendering/SkeletalMeshModel.h"
+#include "Rendering/SkeletalMeshLODModel.h"
+#include "ReferenceSkeleton.h"
+#include "Widgets/SBoxPanel.h"
+#include "SNeuralMorphBoneGroupsWidget.h"
+#include "SNeuralMorphCurveGroupsWidget.h"
+#include "SNeuralMorphInputWidget.h"
 
 #define LOCTEXT_NAMESPACE "NeuralMorphEditorModel"
 
@@ -23,6 +37,15 @@ namespace UE::NeuralMorphModel
 		return new FNeuralMorphEditorModel();
 	}
 
+	void FNeuralMorphEditorModel::RebuildEditorMaskInfo()
+	{
+		UNeuralMorphInputInfo* NeuralInputInfo = Cast<UNeuralMorphInputInfo>(GetEditorInputInfo());
+		if (NeuralInputInfo)
+		{
+			BuildMaskBuffer(NeuralInputInfo->GetInputItemMaskBuffer());
+		}
+	}
+
 	void FNeuralMorphEditorModel::OnPropertyChanged(FPropertyChangedEvent& PropertyChangedEvent)
 	{
 		const FProperty* Property = PropertyChangedEvent.Property;
@@ -42,6 +65,10 @@ namespace UE::NeuralMorphModel
 			Property->GetFName() == TEXT("CurveName"))		// The curve name inside one of the items in the CurveNames list changed.
 		{
 			UpdateIsReadyForTrainingState();
+			if (GetNeuralMorphModel()->GetModelMode() == ENeuralMorphMode::Local)
+			{
+				RebuildEditorMaskInfo();
+			}
 		}
 		else if (Property->GetFName() == GET_MEMBER_NAME_CHECKED(UNeuralMorphModel, Mode))
 		{
@@ -49,10 +76,20 @@ namespace UE::NeuralMorphModel
 			{
 				UpdateIsReadyForTrainingState();
 				GetEditor()->GetModelDetailsView()->ForceRefresh();
+				if (GetNeuralMorphModel()->GetModelMode() == ENeuralMorphMode::Local)
+				{
+					RebuildEditorMaskInfo();
+				}
 			}
 		}
 	}
 
+	TSharedPtr<SMLDeformerInputWidget> FNeuralMorphEditorModel::CreateInputWidget()
+	{
+		return SNew(SNeuralMorphInputWidget)
+			.EditorModel(this);
+	}
+
 	bool FNeuralMorphEditorModel::IsTrained() const
 	{
 		UNeuralMorphNetwork* MorphNetwork = GetNeuralMorphModel()->GetNeuralMorphNetwork();
@@ -64,6 +101,17 @@ namespace UE::NeuralMorphModel
 		return false;
 	}
 
+	void FNeuralMorphEditorModel::Init(const InitSettings& Settings)
+	{
+		FMLDeformerMorphModelEditorModel::Init(Settings);
+
+		RemoveNonExistingMaskInfos();
+		if (GetNeuralMorphModel()->GetModelMode() == ENeuralMorphMode::Local)
+		{
+			RebuildEditorMaskInfo();
+		}
+	}
+
 	ETrainingResult FNeuralMorphEditorModel::Train()
 	{
 		return TrainModel<UNeuralMorphTrainingModel>(this);
@@ -115,29 +163,24 @@ namespace UE::NeuralMorphModel
 			return;
 		}
 
-		const bool bIncludeBones = Model->DoesSupportBones() && Model->ShouldIncludeBonesInTraining();
-		const bool bIncludeCurves = Model->DoesSupportCurves() && Model->ShouldIncludeCurvesInTraining();
 		const USkeleton* Skeleton = Model->GetSkeletalMesh() ? Model->GetSkeletalMesh()->GetSkeleton() : nullptr;
 		const USkeletalMesh* SkeletalMesh = Model->GetSkeletalMesh();
 
 		// Handle bones.
-		if (bIncludeBones && SkeletalMesh)
+		if (SkeletalMesh)
 		{
 			const TArray<FNeuralMorphBoneGroup>& ModelBoneGroups = NeuralMorphModel->GetBoneGroups();
 			const FReferenceSkeleton& RefSkeleton = SkeletalMesh->GetRefSkeleton();
 			for (int32 BoneGroupIndex = 0; BoneGroupIndex < ModelBoneGroups.Num(); ++BoneGroupIndex)
 			{
 				const FNeuralMorphBoneGroup& BoneGroup = ModelBoneGroups[BoneGroupIndex];
-				if (BoneGroup.BoneNames.IsEmpty())
-				{
-					continue;
-				}
 
 				NeuralInputInfo->GetBoneGroups().AddDefaulted();
 				FNeuralMorphBoneGroup& NewGroup = NeuralInputInfo->GetBoneGroups().Last();
 
 				const int32 NumBonesInGroup = BoneGroup.BoneNames.Num();
 				NewGroup.BoneNames.AddDefaulted(NumBonesInGroup);
+				NewGroup.GroupName = BoneGroup.GroupName;
 				for (int32 Index = 0; Index < NumBonesInGroup; ++Index)
 				{
 					const FBoneReference& BoneRef = BoneGroup.BoneNames[Index];
@@ -165,7 +208,7 @@ namespace UE::NeuralMorphModel
 		}
 
 		// Handle curves.
-		if (bIncludeCurves && SkeletalMesh)
+		if (SkeletalMesh)
 		{
 			const TArray<FNeuralMorphCurveGroup>& ModelCurveGroups = NeuralMorphModel->GetCurveGroups();
 			const FReferenceSkeleton& RefSkeleton = SkeletalMesh->GetRefSkeleton();
@@ -182,6 +225,7 @@ namespace UE::NeuralMorphModel
 
 				const int32 NumCurvesInGroup = CurveGroup.CurveNames.Num();
 				NewGroup.CurveNames.AddDefaulted(NumCurvesInGroup);
+				NewGroup.GroupName = CurveGroup.GroupName;
 				for (int32 Index = 0; Index < NumCurvesInGroup; ++Index)
 				{
 					const FMLDeformerCurveReference& CurveRef = CurveGroup.CurveNames[Index];
@@ -201,6 +245,9 @@ namespace UE::NeuralMorphModel
 				}
 			}
 		}
+
+		RemoveNonExistingMaskInfos();
+		BuildMaskBuffer(NeuralInputInfo->GetInputItemMaskBuffer());
 	}
 
 	void FNeuralMorphEditorModel::UpdateIsReadyForTrainingState()
@@ -225,10 +272,474 @@ namespace UE::NeuralMorphModel
 			Text = FText::Format(
 				LOCTEXT("GroupErrorFormat", "{0}\n{1}"), 
 				Text,
-				LOCTEXT("GroupErrorText", "Bone and curve groups require at least two valid items.\nCheck the log warnings for more information."));
+				LOCTEXT("GroupErrorText", "There are invalid bone and/or curve groups.\nCheck the log warnings for more information."));
 		}
 		return Text;
 	}
+
+	const TArrayView<const float> FNeuralMorphEditorModel::GetMaskForMorphTarget(int32 MorphTargetIndex) const
+	{
+		const UNeuralMorphModel* NeuralMorphModel = GetNeuralMorphModel();
+		const UMLDeformerMorphModelInputInfo* MorphInputInfo = Cast<UMLDeformerMorphModelInputInfo>(NeuralMorphModel->GetInputInfo());
+		check(MorphInputInfo);
+
+		const TArray<float>& MaskBuffer = MorphInputInfo->GetInputItemMaskBuffer();
+		if (MaskBuffer.IsEmpty() || !NeuralMorphModel->IsBoneMaskingEnabled())
+		{
+			return TArrayView<const float>();
+		}
+
+		const int32 ItemIndex = MorphTargetIndex / NeuralMorphModel->GetLocalNumMorphsPerBone();
+		return MorphInputInfo->GetMaskForItem(ItemIndex);
+	}
+
+	void FNeuralMorphEditorModel::ApplyMaskInfoToMaskBuffer(const USkeletalMesh* SkeletalMesh, const FNeuralMorphMaskInfo& MaskInfo, TArrayView<float> ItemMaskBuffer)	
+	{
+		const FReferenceSkeleton& RefSkel = SkeletalMesh->GetRefSkeleton();
+
+		// Apply the bones to the mask buffer.
+		const TArray<int32> VirtualParentTable = BuildVirtualParentTable(RefSkel, GetEditorInputInfo()->GetBoneNames());
+		for (const FName MaskBoneName : MaskInfo.BoneNames)
+		{
+			const int32 MaskBoneIndex = RefSkel.FindBoneIndex(MaskBoneName);
+			if (MaskBoneIndex != INDEX_NONE)
+			{
+				ApplyBoneToMask(MaskBoneIndex, ItemMaskBuffer);
+			}
+			else
+			{
+				UE_LOG(LogNeuralMorphModel, Warning, TEXT("Mask contains a bone named '%s', which cannot be found in the ref skeleton of skeletal mesh '%s'."),
+					*MaskBoneName.ToString(),
+					*SkeletalMesh->GetName());
+			}
+
+			// Add required bones
+			TArray<int32> RequiredBones;
+			const int32 SkeletonBoneIndex = RefSkel.FindBoneIndex(MaskBoneName);
+			if (SkeletonBoneIndex != INDEX_NONE)
+			{
+				AddRequiredBones(RefSkel, SkeletonBoneIndex, VirtualParentTable, RequiredBones);
+			}
+			for (const int32 RequiredBoneIndex: RequiredBones)
+			{
+				check(RequiredBoneIndex != INDEX_NONE);
+				const FName RequiredBoneName = RefSkel.GetBoneName(RequiredBoneIndex);
+				if (!MaskInfo.BoneNames.Contains(RequiredBoneName))
+				{
+					ApplyBoneToMask(RequiredBoneIndex, ItemMaskBuffer);
+				}
+			}
+		}
+	}
+
+	void FNeuralMorphEditorModel::BuildMaskBuffer(TArray<float>& OutMaskBuffer)
+	{
+		if (Model->GetSkeletalMesh() == nullptr)
+		{
+			OutMaskBuffer.Empty();
+			return;
+		}
+
+		// Clear the mask buffer.
+		UNeuralMorphModel* NeuralMorphModel = GetNeuralMorphModel();
+		OutMaskBuffer.Empty();
+		if (NeuralMorphModel->GetModelMode() != ENeuralMorphMode::Local)
+		{
+			return;
+		}
+
+		const int32 NumBaseMeshVerts = Model->GetNumBaseMeshVerts();
+		check(NumBaseMeshVerts > 0);
+
+		USkeletalMesh* SkeletalMesh = Model->GetSkeletalMesh();	
+		UNeuralMorphInputInfo* NeuralMorphInputInfo = Cast<UNeuralMorphInputInfo>(GetEditorInputInfo());
+		const FReferenceSkeleton& RefSkel = SkeletalMesh->GetRefSkeleton();
+		const TArray<FName>& BoneNames = NeuralMorphInputInfo->GetBoneNames();
+
+		// Calculate the number of flo =ats we need to store all the masks.
+		const int32 NumBones = NeuralMorphInputInfo->GetNumBones();
+		const int32 NumCurves = NeuralMorphInputInfo->GetNumCurves();
+		const int32 NumBoneGroups = NeuralMorphInputInfo->GetBoneGroups().Num();
+		const int32 NumCurveGroups = NeuralMorphInputInfo->GetCurveGroups().Num();
+		const int32 NumFloats = (NumBones + NumCurves + NumBoneGroups + NumCurveGroups) * NumBaseMeshVerts;
+		const int32 HierarchyDepth = 1;	// Default hierarchy depth in case we have no mask for specific bones yet.
+
+		// Init the mask to all zeros.
+		OutMaskBuffer.SetNumZeroed(NumFloats);
+
+		// For all bones.
+		int32 MaskOffset = 0;
+		for (int32 Index = 0; Index < NumBones; ++Index)
+		{
+			TArrayView<float> ItemMaskBuffer(&OutMaskBuffer[MaskOffset], NumBaseMeshVerts);
+			const FName BoneName = BoneNames[Index];
+			FNeuralMorphMaskInfo* MaskInfo = NeuralMorphModel->BoneMaskInfos.Find(BoneName);
+			if (MaskInfo == nullptr)
+			{
+				MaskInfo = &NeuralMorphModel->BoneMaskInfos.Add(BoneName, FNeuralMorphMaskInfo());
+				GenerateBoneMaskInfo(Index, HierarchyDepth);
+			}
+
+			ApplyMaskInfoToMaskBuffer(SkeletalMesh, *MaskInfo, ItemMaskBuffer);
+			MaskOffset += NumBaseMeshVerts;
+		}
+
+		// For all curves, just use a mask full of values of 1.
+		for (int32 Index = 0; Index < NumCurves; ++Index)
+		{
+			TArrayView<float> ItemMaskBuffer(&OutMaskBuffer[MaskOffset], NumBaseMeshVerts);
+			FillMaskValues(ItemMaskBuffer, 1.0f);
+			MaskOffset += NumBaseMeshVerts;
+		}
+
+		// For all bone groups.		
+		for (int32 Index = 0; Index < NumBoneGroups; ++Index)
+		{
+			TArrayView<float> ItemMaskBuffer(&OutMaskBuffer[MaskOffset], NumBaseMeshVerts);
+			
+			const FName GroupName = NeuralMorphInputInfo->GetBoneGroups()[Index].GroupName;
+			FNeuralMorphMaskInfo* MaskInfo = NeuralMorphModel->BoneGroupMaskInfos.Find(GroupName);
+			if (MaskInfo == nullptr)
+			{
+				MaskInfo = &NeuralMorphModel->BoneGroupMaskInfos.Add(GroupName, FNeuralMorphMaskInfo());
+				GenerateBoneGroupMaskInfo(Index, HierarchyDepth);
+			}
+
+			ApplyMaskInfoToMaskBuffer(SkeletalMesh, *MaskInfo, ItemMaskBuffer);
+			MaskOffset += NumBaseMeshVerts;
+		}
+
+		// For all curve groups, just use a mask full of values of 1.
+		for (int32 Index = 0; Index < NumCurveGroups; ++Index)
+		{
+			TArrayView<float> ItemMaskBuffer(&OutMaskBuffer[MaskOffset], NumBaseMeshVerts);
+			FillMaskValues(ItemMaskBuffer, 1.0f);
+			MaskOffset += NumBaseMeshVerts;
+		}
+	}
+
+	FName FNeuralMorphEditorModel::GenerateUniqueBoneGroupName() const
+	{
+		UNeuralMorphModel* NeuralMorphModel = GetNeuralMorphModel();
+		FName GeneratedName;
+		bool bUniqueNameFound = false;
+		int32 Counter = 0;
+		while (!bUniqueNameFound)
+		{
+			GeneratedName = FName(FString::Format(TEXT("Bone Group #{0}"), {Counter++}));
+
+			bool bHasGroupWithSameName = false;
+			for (const FNeuralMorphBoneGroup& BoneGroup : NeuralMorphModel->GetBoneGroups())
+			{
+				if (BoneGroup.GroupName == GeneratedName)
+				{
+					bHasGroupWithSameName = true;
+					break;
+				}
+			}
+
+			if (!bHasGroupWithSameName)
+			{
+				bUniqueNameFound = true;
+			}
+		}
+
+		return GeneratedName;
+	}
+
+	FName FNeuralMorphEditorModel::GenerateUniqueCurveGroupName() const
+	{
+		UNeuralMorphModel* NeuralMorphModel = GetNeuralMorphModel();
+		FName GeneratedName;
+		bool bUniqueNameFound = false;
+		int32 Counter = 0;
+		while (!bUniqueNameFound)
+		{
+			GeneratedName = FName(FString::Format(TEXT("Curve Group #{0}"), {Counter++}));
+
+			bool bHasGroupWithSameName = false;
+			for (const FNeuralMorphBoneGroup& BoneGroup : NeuralMorphModel->GetBoneGroups())
+			{
+				if (BoneGroup.GroupName == GeneratedName)
+				{
+					bHasGroupWithSameName = true;
+					break;
+				}
+			}
+
+			if (!bHasGroupWithSameName)
+			{
+				bUniqueNameFound = true;
+			}
+		}
+
+		return GeneratedName;
+	}
+
+	void FNeuralMorphEditorModel::ResetBoneMaskInfos()
+	{
+		GetNeuralMorphModel()->BoneMaskInfos.Empty();
+	}
+
+	void FNeuralMorphEditorModel::ResetBoneGroupMaskInfos()
+	{
+		GetNeuralMorphModel()->BoneGroupMaskInfos.Empty();
+	}
+
+	void FNeuralMorphEditorModel::GenerateBoneMaskInfo(int32 InputInfoBoneIndex, int32 HierarchyDepth)
+	{
+		UNeuralMorphModel* NeuralMorphModel = GetNeuralMorphModel();
+		USkeletalMesh* SkeletalMesh = Model->GetSkeletalMesh();	
+		UNeuralMorphInputInfo* NeuralMorphInputInfo = Cast<UNeuralMorphInputInfo>(GetEditorInputInfo());
+		const FReferenceSkeleton& RefSkel = SkeletalMesh->GetRefSkeleton();
+		const TArray<FName>& BoneNames = NeuralMorphInputInfo->GetBoneNames();
+		const TArray<int32> VirtualParentTable = BuildVirtualParentTable(RefSkel, BoneNames);
+		const int32 NumBones = NeuralMorphInputInfo->GetNumBones();
+
+		// Make sure we have a valid bone name.
+		const FName BoneName = BoneNames[InputInfoBoneIndex];
+		if (BoneName == NAME_None)
+		{
+			return;
+		}
+
+		// Get the bone index inside our skeleton.
+		const int32 SkeletonBoneIndex = RefSkel.FindBoneIndex(BoneName);
+		if (SkeletonBoneIndex == INDEX_NONE)
+		{
+			UE_LOG(LogNeuralMorphModel, Warning, TEXT("Bone '%s' cannot be found in the SkeletalMesh '%s', ignoring during mask generation."), *BoneName.ToString(), *SkeletalMesh->GetName());
+			return;
+		}
+
+		// Build the list of bones that are required for the mask.
+		TArray<int32> BonesAdded;
+		RecursiveAddBoneToMaskUpwards(RefSkel, SkeletonBoneIndex, HierarchyDepth, VirtualParentTable, BonesAdded);
+		RecursiveAddBoneToMaskDownwards(RefSkel, SkeletonBoneIndex, HierarchyDepth, VirtualParentTable, BonesAdded);
+
+		// Now that we know which bones we need, add them to the mask.
+		FNeuralMorphMaskInfo* MaskInfo = NeuralMorphModel->BoneMaskInfos.Find(BoneName);
+		if (MaskInfo == nullptr)
+		{
+			MaskInfo = &NeuralMorphModel->BoneMaskInfos.Add(BoneName, FNeuralMorphMaskInfo());
+		}
+		MaskInfo->BoneNames.Reset();
+		for (int32 BoneIndex : BonesAdded)
+		{
+			const FName MaskBoneName = RefSkel.GetBoneName(BoneIndex);
+			MaskInfo->BoneNames.Add(MaskBoneName);
+		}
+	}
+
+	void FNeuralMorphEditorModel::GenerateBoneGroupMaskInfo(int32 InputInfoBoneGroupIndex, int32 HierarchyDepth)
+	{
+		UNeuralMorphModel* NeuralMorphModel = GetNeuralMorphModel();
+		USkeletalMesh* SkeletalMesh = Model->GetSkeletalMesh();	
+		UNeuralMorphInputInfo* NeuralMorphInputInfo = Cast<UNeuralMorphInputInfo>(GetEditorInputInfo());
+		const FReferenceSkeleton& RefSkel = SkeletalMesh->GetRefSkeleton();
+		const TArray<FName>& BoneNames = NeuralMorphInputInfo->GetBoneNames();
+		const TArray<int32> VirtualParentTable = BuildVirtualParentTable(RefSkel, BoneNames);
+
+		// For all bones.
+		TArray<int32> BonesAdded;
+		const FNeuralMorphBoneGroup& BoneGroup = NeuralMorphModel->GetBoneGroups()[InputInfoBoneGroupIndex];
+		for (int32 BoneIndex = 0; BoneIndex < BoneGroup.BoneNames.Num(); ++BoneIndex)
+		{
+			// Get the bone name.
+			const FName BoneName = BoneGroup.BoneNames[BoneIndex].BoneName;
+			if (BoneName == NAME_None)
+			{
+				UE_LOG(LogNeuralMorphModel, Warning, TEXT("Bone index %d inside bone group %d has its name set to None, ignoring the bone inside the group's mask."), BoneIndex, InputInfoBoneGroupIndex);
+				continue;
+			}
+
+			// Get the bone index inside our skeleton.
+			const int32 SkeletonBoneIndex = RefSkel.FindBoneIndex(BoneName);
+			if (SkeletonBoneIndex == INDEX_NONE)
+			{
+				UE_LOG(LogNeuralMorphModel, Warning, TEXT("Bone '%s' cannot be found in the SkeletalMesh '%s', ignoring the bone inside the group's mask."), *BoneName.ToString(), *SkeletalMesh->GetName());
+				continue;
+			}
+
+			// Add all bones up and down the chain.
+			RecursiveAddBoneToMaskUpwards(RefSkel, SkeletonBoneIndex, HierarchyDepth, VirtualParentTable, BonesAdded);
+			RecursiveAddBoneToMaskDownwards(RefSkel, SkeletonBoneIndex, HierarchyDepth, VirtualParentTable, BonesAdded);
+		}
+
+		// Now that we know which bones we need, add them to the mask.
+		FNeuralMorphMaskInfo* MaskInfo = NeuralMorphModel->BoneGroupMaskInfos.Find(BoneGroup.GroupName);
+		if (MaskInfo == nullptr)
+		{
+			MaskInfo = &NeuralMorphModel->BoneGroupMaskInfos.Add(BoneGroup.GroupName, FNeuralMorphMaskInfo());
+		}
+		MaskInfo->BoneNames.Reset();
+		for (int32 BoneIndex : BonesAdded)
+		{
+			const FName MaskBoneName = RefSkel.GetBoneName(BoneIndex);
+			MaskInfo->BoneNames.Add(MaskBoneName);
+		}
+	}
+
+	void FNeuralMorphEditorModel::GenerateBoneMaskInfos(int32 HierarchyDepth)
+	{
+		check(HierarchyDepth >= 1);
+		ResetBoneMaskInfos();
+
+		// Generate a bone mask info for each bone.
+		UNeuralMorphModel* NeuralMorphModel = GetNeuralMorphModel();
+		UNeuralMorphInputInfo* NeuralMorphInputInfo = Cast<UNeuralMorphInputInfo>(GetEditorInputInfo());
+		const int32 NumBones = NeuralMorphInputInfo->GetNumBones();
+		for (int32 Index = 0; Index < NumBones; ++Index)
+		{
+			GenerateBoneMaskInfo(Index, HierarchyDepth);
+		}
+	}
+
+	void FNeuralMorphEditorModel::GenerateBoneGroupMaskInfos(int32 HierarchyDepth)
+	{
+		check(HierarchyDepth >= 1);
+
+		ResetBoneGroupMaskInfos();
+
+		// Generate a mask info for each group.
+		UNeuralMorphModel* NeuralMorphModel = GetNeuralMorphModel();
+		UNeuralMorphInputInfo* NeuralMorphInputInfo = Cast<UNeuralMorphInputInfo>(GetEditorInputInfo());
+		const int32 NumBoneGroups = NeuralMorphInputInfo->GetBoneGroups().Num();	
+		for (int32 Index = 0; Index < NumBoneGroups; ++Index)
+		{
+			GenerateBoneGroupMaskInfo(Index, HierarchyDepth);
+		}
+	}
+
+	void FNeuralMorphEditorModel::OnPostInputAssetChanged()
+	{
+		FMLDeformerMorphModelEditorModel::OnPostInputAssetChanged();
+		RemoveNonExistingMaskInfos();
+		if (GetNeuralMorphModel()->GetModelMode() == ENeuralMorphMode::Local)
+		{
+			RebuildEditorMaskInfo();
+		}
+	}
+
+	void FNeuralMorphEditorModel::RemoveNonExistingMaskInfos()
+	{
+		UNeuralMorphModel* NeuralMorphModel = GetNeuralMorphModel();
+
+		// Find all mask infos for bones that don't exist anymore.
+		TArray<FName> BoneMaskInfosToRemove;
+		for (const auto& MaskInfo : NeuralMorphModel->BoneMaskInfos)
+		{
+			const FName BoneName = MaskInfo.Key;
+			if (!NeuralMorphModel->GetBoneIncludeList().Contains(BoneName))
+			{
+				BoneMaskInfosToRemove.Add(BoneName);
+			}
+		}
+
+		// Remove those mask infos.
+		for (const FName BoneName : BoneMaskInfosToRemove)
+		{
+			NeuralMorphModel->BoneMaskInfos.Remove(BoneName);
+		}
+
+		// Find all bone groups that don't exist anymore.
+		BoneMaskInfosToRemove.Reset();
+		for (const auto& MaskInfo : NeuralMorphModel->BoneGroupMaskInfos)
+		{
+			const FName GroupName = MaskInfo.Key;
+
+			bool bFound = false;
+			for (const FNeuralMorphBoneGroup& BoneGroup : NeuralMorphModel->GetBoneGroups())
+			{
+				if (BoneGroup.GroupName == GroupName)
+				{
+					bFound = true;
+					break;
+				}
+			}
+
+			if (!bFound)
+			{
+				BoneMaskInfosToRemove.Add(GroupName);
+			}
+		}
+
+		// Remove those group mask infos.
+		for (const FName GroupName : BoneMaskInfosToRemove)
+		{
+			NeuralMorphModel->BoneGroupMaskInfos.Remove(GroupName);
+		}
+	}
+
+	void FNeuralMorphEditorModel::DebugDrawItemMask(FPrimitiveDrawInterface* PDI, int32 MaskItemIndex, const FVector& DrawOffset)
+	{
+		UMLDeformerMorphModelVizSettings* VizSettings = Cast<UMLDeformerMorphModelVizSettings>(Model->GetVizSettings());
+		const int32 NumVerts = Model->GetNumBaseMeshVerts();
+		const TArray<FVector3f>& UnskinnedPositions = Sampler->GetUnskinnedVertexPositions();
+		if (MaskVizItemIndex == INDEX_NONE ||
+			NumVerts != Model->GetInputInfo()->GetNumBaseMeshVertices() ||
+			UnskinnedPositions.Num() != NumVerts)
+		{
+			return;
+		}
+
+		UMLDeformerMorphModel* MorphModel = GetMorphModel();
+		UNeuralMorphInputInfo* InputInfo = Cast<UNeuralMorphInputInfo>(GetEditorInputInfo());
+		const int32 NumMaskItems = InputInfo->GetBoneNames().Num() + InputInfo->GetCurveNames().Num() + InputInfo->GetBoneGroups().Num() + InputInfo->GetCurveGroups().Num();
+		const int32 FinalMaskItemIndex = FMath::Clamp<int32>(MaskItemIndex, 0, NumMaskItems - 1);
+		const TArrayView<const float> MaskBuffer = InputInfo->GetMaskForItem(FinalMaskItemIndex);
+		if (!MaskBuffer.IsEmpty())
+		{
+			check(MaskBuffer.Num() == NumVerts);
+			check(NumVerts == UnskinnedPositions.Num());
+
+			FLinearColor OrgIncludedColor = FMLDeformerEditorStyle::Get().GetColor("MLDeformer.Morphs.MaskIncludedVertexColor");
+			const FLinearColor ExcludedColor = FMLDeformerEditorStyle::Get().GetColor("MLDeformer.Morphs.MaskExcludedVertexColor");
+			for (int32 VertexIndex = 0; VertexIndex < NumVerts; ++VertexIndex)
+			{
+				const FVector StartPoint = FVector(UnskinnedPositions[VertexIndex]) + DrawOffset;
+				if (MaskBuffer[VertexIndex] > 0.0f)
+				{
+					const FLinearColor IncludedColor = OrgIncludedColor * FMath::Clamp<float>(MaskBuffer[VertexIndex], 0.0f, 1.0f);
+					PDI->DrawPoint(StartPoint, IncludedColor, 1.0f, 0);
+				}
+				else
+				{
+					PDI->DrawPoint(StartPoint, ExcludedColor, 0.75f, 0);
+				}
+			}
+		}
+	}
+
+	void FNeuralMorphEditorModel::Render(const FSceneView* View, FViewport* Viewport, FPrimitiveDrawInterface* PDI)
+	{
+		FMLDeformerMorphModelEditorModel::Render(View, Viewport, PDI);
+
+		const UNeuralMorphModelVizSettings* VizSettings = Cast<UNeuralMorphModelVizSettings>(Model->GetVizSettings());
+		check(VizSettings);
+		if (VizSettings->GetVisualizationMode() == EMLDeformerVizMode::TestData &&
+			VizSettings->MaskVizMode != ENeuralMorphMaskVizMode::Off &&
+			GetNeuralMorphModel()->GetModelMode() == ENeuralMorphMode::Local &&
+			MaskVizItemIndex != INDEX_NONE)
+		{
+			bool bDrawMask = true;
+			if (VizSettings->MaskVizMode == ENeuralMorphMaskVizMode::WhenInFocus && !InputWidget->HasFocusedDescendants())
+			{
+				bDrawMask = false;
+			}
+
+			if (bDrawMask)
+			{
+				FVector DrawOffset = -VizSettings->GetMeshSpacingOffsetVector();
+				if (VizSettings->GetDrawMorphTargets() && !GetMorphModel()->GetMorphTargetDeltas().IsEmpty())
+				{
+					DrawOffset *= 2.0f;
+				}
+
+				DebugDrawItemMask(PDI, MaskVizItemIndex, DrawOffset);
+			}
+		}
+	}
+
 }	// namespace UE::NeuralMorphModel
 
 #undef LOCTEXT_NAMESPACE
diff --git a/Engine/Plugins/Animation/MLDeformer/NeuralMorphModel/Source/NeuralMorphModelEditor/Private/NeuralMorphModelDetails.cpp b/Engine/Plugins/Animation/MLDeformer/NeuralMorphModel/Source/NeuralMorphModelEditor/Private/NeuralMorphModelDetails.cpp
--- a/Engine/Plugins/Animation/MLDeformer/NeuralMorphModel/Source/NeuralMorphModelEditor/Private/NeuralMorphModelDetails.cpp
+++ b/Engine/Plugins/Animation/MLDeformer/NeuralMorphModel/Source/NeuralMorphModelEditor/Private/NeuralMorphModelDetails.cpp
@@ -2,14 +2,9 @@
 
 #include "NeuralMorphModelDetails.h"
 #include "NeuralMorphModel.h"
-#include "NeuralMorphInputInfo.h"
 #include "NeuralMorphEditorModel.h"
 #include "DetailCategoryBuilder.h"
 #include "DetailWidgetRow.h"
-#include "Widgets/Layout/SBox.h"
-#include "Widgets/Input/SButton.h"
-#include "SWarningOrErrorBox.h"
-
 
 #define LOCTEXT_NAMESPACE "NeuralMorphModelDetails"
 
@@ -53,20 +48,7 @@ namespace UE::NeuralMorphModel
 		TrainingSettingsCategoryBuilder->AddProperty(GET_MEMBER_NAME_CHECKED(UNeuralMorphModel, LearningRate), UNeuralMorphModel::StaticClass());
 		TrainingSettingsCategoryBuilder->AddProperty(GET_MEMBER_NAME_CHECKED(UNeuralMorphModel, LearningRateDecay), UNeuralMorphModel::StaticClass());
 		TrainingSettingsCategoryBuilder->AddProperty(GET_MEMBER_NAME_CHECKED(UNeuralMorphModel, RegularizationFactor), UNeuralMorphModel::StaticClass());
-	}
-
-	void FNeuralMorphModelDetails::AddTrainingInputFilters()
-	{
-		FMLDeformerMorphModelDetails::AddTrainingInputFilters();
-
-		UNeuralMorphModel* NeuralMorphModel = Cast<UNeuralMorphModel>(Model);
-		check(NeuralMorphModel);
-
-		InputOutputCategoryBuilder->AddProperty(GET_MEMBER_NAME_CHECKED(UNeuralMorphModel, BoneGroups), UNeuralMorphModel::StaticClass())
-			.Visibility(NeuralMorphModel->Mode == ENeuralMorphMode::Local ? EVisibility::Visible : EVisibility::Collapsed);
-
-		InputOutputCategoryBuilder->AddProperty(GET_MEMBER_NAME_CHECKED(UNeuralMorphModel, CurveGroups), UNeuralMorphModel::StaticClass())
-			.Visibility(NeuralMorphModel->Mode == ENeuralMorphMode::Local ? EVisibility::Visible : EVisibility::Collapsed);
+		TrainingSettingsCategoryBuilder->AddProperty(GET_MEMBER_NAME_CHECKED(UNeuralMorphModel, bEnableBoneMasks), UNeuralMorphModel::StaticClass());
 	}
 }	// namespace UE::NeuralMorphModel
 
diff --git a/Engine/Plugins/Animation/MLDeformer/NeuralMorphModel/Source/NeuralMorphModelEditor/Private/NeuralMorphModelEditorModule.cpp b/Engine/Plugins/Animation/MLDeformer/NeuralMorphModel/Source/NeuralMorphModelEditor/Private/NeuralMorphModelEditorModule.cpp
--- a/Engine/Plugins/Animation/MLDeformer/NeuralMorphModel/Source/NeuralMorphModelEditor/Private/NeuralMorphModelEditorModule.cpp
+++ b/Engine/Plugins/Animation/MLDeformer/NeuralMorphModel/Source/NeuralMorphModelEditor/Private/NeuralMorphModelEditorModule.cpp
@@ -4,6 +4,9 @@
 #include "NeuralMorphModelDetails.h"
 #include "NeuralMorphEditorModel.h"
 #include "NeuralMorphModel.h"
+#include "SNeuralMorphInputWidget.h"
+#include "SNeuralMorphBoneGroupsWidget.h"
+#include "SNeuralMorphCurveGroupsWidget.h"
 #include "MLDeformerEditorModule.h"
 #include "Modules/ModuleManager.h"
 #include "EditorModeRegistry.h"
@@ -41,6 +44,10 @@ namespace UE::NeuralMorphModel
 		FMLDeformerEditorModule& EditorModule = FModuleManager::LoadModuleChecked<FMLDeformerEditorModule>("MLDeformerFrameworkEditor");
 		FMLDeformerEditorModelRegistry& ModelRegistry = EditorModule.GetModelRegistry();
 		ModelRegistry.RegisterEditorModel(UNeuralMorphModel::StaticClass(), FOnGetEditorModelInstance::CreateStatic(&FNeuralMorphEditorModel::MakeInstance), /*ModelPriority*/100);
+
+		FNeuralMorphInputWidgetCommands::Register();
+		FNeuralMorphBoneGroupsCommands::Register();
+		FNeuralMorphCurveGroupsCommands::Register();
 	}
 
 	void FNeuralMorphModelEditorModule::ShutdownModule()
diff --git a/Engine/Plugins/Animation/MLDeformer/NeuralMorphModel/Source/NeuralMorphModelEditor/Private/NeuralMorphTrainingModel.cpp b/Engine/Plugins/Animation/MLDeformer/NeuralMorphModel/Source/NeuralMorphModelEditor/Private/NeuralMorphTrainingModel.cpp
--- a/Engine/Plugins/Animation/MLDeformer/NeuralMorphModel/Source/NeuralMorphModelEditor/Private/NeuralMorphTrainingModel.cpp
+++ b/Engine/Plugins/Animation/MLDeformer/NeuralMorphModel/Source/NeuralMorphModelEditor/Private/NeuralMorphTrainingModel.cpp
@@ -42,3 +42,21 @@ TArray<int32> UNeuralMorphTrainingModel::GenerateCurveGroupIndices() const
 	InputInfo->GenerateCurveGroupIndices(OutCurveGroupIndices);
 	return MoveTemp(OutCurveGroupIndices);
 }
+
+TArray<float> UNeuralMorphTrainingModel::GetMorphTargetMasks() const
+{
+	using namespace UE::NeuralMorphModel;
+	FNeuralMorphEditorModel* NeuralMorphEditorModel = static_cast<FNeuralMorphEditorModel*>(EditorModel);
+	UNeuralMorphModel* NeuralMorphModel = NeuralMorphEditorModel->GetNeuralMorphModel();
+	UNeuralMorphInputInfo* NeuralInputInfo = Cast<UNeuralMorphInputInfo>(NeuralMorphEditorModel->GetEditorInputInfo());
+	if (NeuralMorphModel->GetModelMode() != ENeuralMorphMode::Local || !NeuralMorphModel->IsBoneMaskingEnabled())
+	{
+		NeuralInputInfo->GetInputItemMaskBuffer().Empty();
+	}
+	else	
+	{
+		NeuralMorphEditorModel->RebuildEditorMaskInfo();
+	}
+
+	return NeuralInputInfo->GetInputItemMaskBuffer();
+}
diff --git a/Engine/Plugins/Animation/MLDeformer/NeuralMorphModel/Source/NeuralMorphModelEditor/Public/NeuralMorphEditorModel.h b/Engine/Plugins/Animation/MLDeformer/NeuralMorphModel/Source/NeuralMorphModelEditor/Public/NeuralMorphEditorModel.h
--- a/Engine/Plugins/Animation/MLDeformer/NeuralMorphModel/Source/NeuralMorphModelEditor/Public/NeuralMorphEditorModel.h
+++ b/Engine/Plugins/Animation/MLDeformer/NeuralMorphModel/Source/NeuralMorphModelEditor/Public/NeuralMorphEditorModel.h
@@ -8,6 +8,11 @@
 
 class UNeuralMorphModel;
 
+namespace UE::MLDeformer
+{
+	class SMLDeformerInputWidget;
+}
+
 namespace UE::NeuralMorphModel
 {
 	using namespace UE::MLDeformer;
@@ -29,14 +34,134 @@ namespace UE::NeuralMorphModel
 		// ~END FGCObject overrides.
 	
 		// FMLDeformerEditorModel overrides.
+		virtual void Init(const InitSettings& Settings) override;
 		virtual ETrainingResult Train() override;
 		virtual bool IsTrained() const override;
 		virtual void InitInputInfo(UMLDeformerInputInfo* InputInfo) override;
 		virtual FText GetOverlayText() const override;
 		virtual void OnPropertyChanged(FPropertyChangedEvent& PropertyChangedEvent) override;
+		virtual TSharedPtr<SMLDeformerInputWidget> CreateInputWidget() override;
+		virtual void OnPostInputAssetChanged() override;
+		virtual void Render(const FSceneView* View, FViewport* Viewport, FPrimitiveDrawInterface* PDI) override;
 		// ~END FMLDeformerEditorModel overrides.
 
+		// FMLDeformerMorphModelEditorModel overrides.
+		virtual const TArrayView<const float> GetMaskForMorphTarget(int32 MorphTargetIndex) const override;
+		virtual bool IsInputMaskingSupported() const override	{ return true; }
+		// ~FMLDeformerMorphModelEditorModel overrides.
+
 		/** Get a pointer to the neural morph runtime model. */
 		UNeuralMorphModel* GetNeuralMorphModel() const			{ return Cast<UNeuralMorphModel>(Model); }
+
+		/**
+		 * Build a set of masks, all concatenated into one flat array.
+		 * The number of masks is (NumBones + NumCurves + NumBoneGroups + NumCurveGroups).
+		 * Each mask contains Model->GetNumBaseMeshVerts() number of floats.
+		 * The masks for curves and curve groups are all set to 1.0.
+		 * The masks for bones and bone groups contain all the vertices they influence during skinning, as well as the ones from the parent bone.
+		 * This will modify the item mask buffer as returned by UMLDeformerMorphModel::GetInputItemMaskBuffer().
+		 * @param OutMaskBuffer The mask buffer to write to. This will automatically be resized.
+		 */
+		void BuildMaskBuffer(TArray<float>& OutMaskBuffer);
+
+		/** Rebuild the mask buffer. */
+		void RebuildEditorMaskInfo();
+
+		/**
+		 * Generate a unique name for a bone group.
+		 * @return The unique name.
+		 */
+		FName GenerateUniqueBoneGroupName() const;
+
+		/**
+		 * Generate a unique name for a curve group.
+		 * @return The unique name.
+		 */
+		FName GenerateUniqueCurveGroupName() const;
+
+		/**
+		 * Remove mask infos for bones that do not exist in the bone include list.
+		 */
+		void RemoveNonExistingMaskInfos();
+
+		/** Clear all bone mask infos. */
+		void ResetBoneMaskInfos();
+
+		/** Clear all bone group mask infos. */
+		void ResetBoneGroupMaskInfos();
+
+		/**
+		 * Generate bone mask infos for all bones in the bone input list.
+		 * Use a specific hierarchy depth for this, which defines how deep up and down the hierarchy we should traverse for each specific bone.
+		 * Imagine we generate the mask for a bone in a chain of bones. We look at all vertices skinned to this bone, and include those vertices in the mask.
+		 * With a hierarchy depth of 1, include all vertices skinned to its parent bone as well as to its child bone.
+		 * With a hierarchy depth of 2 we also include the parent of the parent and the child of the child if it exists.
+		 * We always want the mask to be large enough, otherwise if say skin stretches a lot, we will not capture the vertex deltas for those vertices far away from the bone, which can cause issues.
+		 * However we don't want the mask to be too large either, as that costs performance and memory usage.
+		 * @param HierarchyDepth The hierarchy depth to use, which should be a value of 1 or above.
+		 */
+		void GenerateBoneMaskInfos(int32 HierarchyDepth);
+
+		/**
+		 * Generate the mask infos for all bone groups. 
+		 * Look at the GenerateBoneMaskInfos method for a detailed description of the HierarchyDepth.
+		 * @param HierarchyDepth The hierarchy depth to use, which should be a value of 1 or above.
+		 * @see GenerateBoneMaskInfos
+		 */
+		void GenerateBoneGroupMaskInfos(int32 HierarchyDepth);
+
+		/**
+		 * Generate the mask info for a given bone.
+		 * Look at the GenerateBoneMaskInfos method for a detailed description of the HierarchyDepth.
+		 * @param InputInfoBoneIndex The bone index inside the input info object.
+		 * @param HierarchyDepth The hierarchy depth to use, which should be a value of 1 or above.
+		 * @see GenerateBoneMaskInfos
+		 */
+		void GenerateBoneMaskInfo(int32 InputInfoBoneIndex, int32 HierarchyDepth);
+
+		/**
+		 * Generate the mask info for a given bone group.
+		 * Look at the GenerateBoneMaskInfos method for a detailed description of the HierarchyDepth.
+		 * @param InputInfoBoneIndex The bone index inside the input info object.
+		 * @param HierarchyDepth The hierarchy depth to use, which should be a value of 1 or above.
+		 * @see GenerateBoneMaskInfos
+		 */
+		void GenerateBoneGroupMaskInfo(int32 InputInfoBoneGroupIndex, int32 HierarchyDepth);
+
+		/**
+		 * Apply the mask info to the mask buffer of floats.
+		 * @param SkeletalMesh The skeletal mesh that the mask relates to.
+		 * @param MaskInfo The mask info object to apply to the float mask buffer.
+		 * @param ItemMaskBuffer The mask buffer for this item. This should be an array view of size NumBaseMeshVerts.
+		 */
+		void ApplyMaskInfoToMaskBuffer(const USkeletalMesh* SkeletalMesh, const FNeuralMorphMaskInfo& MaskInfo, TArrayView<float> ItemMaskBuffer);
+
+		/**
+		 * Set the mask visualization item index. This specifies which item (bone, curve, bone group, curve group) to visualize the mask for.
+		 * The first index values are for all input bones. Indices that come after that are for curves, then for bone groups and finally followed by curve groups.
+		 * So if there are 10 input bones and 4 bone groups, and the visualization index is 11, it means it's for the second bone group.
+		 * @param InMaskVizItemIndex The mask visualization index.
+		 */
+		void SetMaskVisualizationItemIndex(int32 InMaskVizItemIndex)	{ MaskVizItemIndex = InMaskVizItemIndex; }
+
+		/**
+		 * Get the mask visualization item index. This specifies which item (bone, curve, bone group, curve group) to visualize the mask for.
+		 * The first index values are for all input bones. Indices that come after that are for curves, then for bone groups and finally followed by curve groups.
+		 * So if there are 10 input bones and 4 bone groups, and the visualization index is 11, it means it's for the second bone group.
+		 * @return The item to visualize the mask for.
+		 */
+		int32 GetMaskVisualizationItemIndex() const						{ return MaskVizItemIndex; }
+
+		/**
+		 * Draw the mask for a given item. An item can be a bone, curve or anything else that has a mask.
+		 * @param PDI The debug draw interface.
+		 * @param MaskItemIndex The item index. The valid range on default is [0..NumEditorInfoBones + NumEditorInfoCurves - 1]. The curves are attached after the bones.
+		 *        So if you have 10 bones, and 5 curves, and the MaskItemIndex is value 12, it will draw the mask for the third curve, because items [0..9] would be for the bones, and [10..14] for the curves.
+		 * @param DrawOffset The offset in world space units to draw the mask at.
+		 */
+		void DebugDrawItemMask(FPrimitiveDrawInterface* PDI, int32 MaskItemIndex, const FVector& DrawOffset);
+
+	protected:
+		int32 MaskVizItemIndex = INDEX_NONE;
 	};
 }	// namespace UE::NeuralMorphModel
diff --git a/Engine/Plugins/Animation/MLDeformer/NeuralMorphModel/Source/NeuralMorphModelEditor/Public/NeuralMorphModelDetails.h b/Engine/Plugins/Animation/MLDeformer/NeuralMorphModel/Source/NeuralMorphModelEditor/Public/NeuralMorphModelDetails.h
--- a/Engine/Plugins/Animation/MLDeformer/NeuralMorphModel/Source/NeuralMorphModelEditor/Public/NeuralMorphModelDetails.h
+++ b/Engine/Plugins/Animation/MLDeformer/NeuralMorphModel/Source/NeuralMorphModelEditor/Public/NeuralMorphModelDetails.h
@@ -23,9 +23,5 @@ namespace UE::NeuralMorphModel
 		// ILayoutDetails overrides.
 		virtual void CustomizeDetails(IDetailLayoutBuilder& DetailBuilder) override;
 		// ~END ILayoutDetails overrides.
-
-		// FMLDeformerModelDetails overrides.
-		virtual void AddTrainingInputFilters() override;
-		// ~END FMLDeformerModelDetails overrides.
 	};
 }	// namespace UE::NeuralMorphModel
diff --git a/Engine/Plugins/Animation/MLDeformer/NeuralMorphModel/Source/NeuralMorphModelEditor/Public/NeuralMorphTrainingModel.h b/Engine/Plugins/Animation/MLDeformer/NeuralMorphModel/Source/NeuralMorphModelEditor/Public/NeuralMorphTrainingModel.h
--- a/Engine/Plugins/Animation/MLDeformer/NeuralMorphModel/Source/NeuralMorphModelEditor/Public/NeuralMorphTrainingModel.h
+++ b/Engine/Plugins/Animation/MLDeformer/NeuralMorphModel/Source/NeuralMorphModelEditor/Public/NeuralMorphTrainingModel.h
@@ -37,4 +37,7 @@ public:
 
 	UFUNCTION(BlueprintPure, Category = "Training Model")
 	TArray<int32> GenerateCurveGroupIndices() const;
+
+	UFUNCTION(BlueprintPure, Category = "Training Model")
+	TArray<float> GetMorphTargetMasks() const;
 };
diff --git a/Engine/Plugins/Animation/MLDeformer/NeuralMorphModel/Source/NeuralMorphModelEditor/Public/SNeuralMorphBoneGroupsWidget.h b/Engine/Plugins/Animation/MLDeformer/NeuralMorphModel/Source/NeuralMorphModelEditor/Public/SNeuralMorphBoneGroupsWidget.h
new file mode 100644
--- /dev/null
+++ b/Engine/Plugins/Animation/MLDeformer/NeuralMorphModel/Source/NeuralMorphModelEditor/Public/SNeuralMorphBoneGroupsWidget.h
@@ -0,0 +1,113 @@
+// Copyright Epic Games, Inc. All Rights Reserved.
+
+#pragma once
+
+#include "CoreTypes.h"
+#include "Templates/SharedPointer.h"
+#include "Widgets/SCompoundWidget.h"
+#include "Widgets/Views/STreeView.h"
+#include "Framework/Commands/Commands.h"
+
+class SWidget;
+
+namespace UE::MLDeformer
+{
+	class SMLDeformerInputWidget;
+}
+
+namespace UE::NeuralMorphModel
+{
+	class FNeuralMorphEditorModel;
+	class SNeuralMorphBoneGroupsWidget;
+	class SNeuralMorphInputWidget;
+	class FNeuralMorphBoneGroupsCommands
+		: public TCommands<FNeuralMorphBoneGroupsCommands>
+	{
+	public:
+		FNeuralMorphBoneGroupsCommands();
+		virtual void RegisterCommands() override;
+
+	public:
+		TSharedPtr<FUICommandInfo> CreateGroup;
+		TSharedPtr<FUICommandInfo> DeleteSelectedItems;
+		TSharedPtr<FUICommandInfo> ClearGroups;
+		TSharedPtr<FUICommandInfo> AddBoneToGroup;
+	};
+
+
+	class FNeuralMorphBoneGroupsTreeElement
+		: public TSharedFromThis<FNeuralMorphBoneGroupsTreeElement>
+	{
+	public:
+		TSharedRef<ITableRow> MakeTreeRowWidget(const TSharedRef<STableViewBase>& InOwnerTable, TSharedRef<FNeuralMorphBoneGroupsTreeElement> InTreeElement, TSharedPtr<SNeuralMorphBoneGroupsWidget> InTreeWidget);
+
+		bool IsGroup() const { return GroupIndex != INDEX_NONE; }
+
+	public:
+		FName Name;
+		TArray<TSharedPtr<FNeuralMorphBoneGroupsTreeElement>> Children;
+		TWeakPtr<FNeuralMorphBoneGroupsTreeElement> ParentGroup;
+		FSlateColor TextColor;
+		int32 GroupIndex = INDEX_NONE;
+		int32 GroupBoneIndex = INDEX_NONE;
+	};
+
+
+	class SNeuralMorphBoneGroupsTreeRowWidget 
+		: public STableRow<TSharedPtr<FNeuralMorphBoneGroupsTreeElement>>
+	{
+	public:
+		void Construct(const FArguments& InArgs, const TSharedRef<STableViewBase>& OwnerTable, TSharedRef<FNeuralMorphBoneGroupsTreeElement> InTreeElement, TSharedPtr<SNeuralMorphBoneGroupsWidget> InTreeView);
+
+	private:
+		TWeakPtr<FNeuralMorphBoneGroupsTreeElement> WeakTreeElement;
+		FText GetName() const;
+
+		friend class SNeuralMorphBoneGroupsWidget; 
+	};
+
+
+	class NEURALMORPHMODELEDITOR_API SNeuralMorphBoneGroupsWidget
+		: public STreeView<TSharedPtr<FNeuralMorphBoneGroupsTreeElement>>
+	{
+	public:
+		SLATE_BEGIN_ARGS(SNeuralMorphBoneGroupsWidget) {}
+		SLATE_ARGUMENT(FNeuralMorphEditorModel*, EditorModel)
+		SLATE_ARGUMENT(TSharedPtr<UE::MLDeformer::SMLDeformerInputWidget>, InputWidget)
+		SLATE_END_ARGS()
+
+		void Construct(const FArguments& InArgs);
+		void BindCommands(TSharedPtr<FUICommandList> CommandList);
+		void Refresh();
+
+		FText GetSectionTitle() const;
+		int32 GetNumSelectedGroups() const;
+		TSharedPtr<UE::MLDeformer::SMLDeformerInputWidget> GetInputWidget() const;
+
+	private:
+		void AddElement(TSharedPtr<FNeuralMorphBoneGroupsTreeElement> Element, TSharedPtr<FNeuralMorphBoneGroupsTreeElement> ParentElement);	
+		void HandleGetChildrenForTree(TSharedPtr<FNeuralMorphBoneGroupsTreeElement> InItem, TArray<TSharedPtr<FNeuralMorphBoneGroupsTreeElement>>& OutChildren);	
+		void OnSelectionChanged(TSharedPtr<FNeuralMorphBoneGroupsTreeElement> Selection, ESelectInfo::Type SelectInfo);
+		const TArray<TSharedPtr<FNeuralMorphBoneGroupsTreeElement>>& GetRootElements() const;
+		FReply OnKeyDown(const FGeometry& InGeometry, const FKeyEvent& InKeyEvent) override;
+		TSharedRef<ITableRow> MakeTableRowWidget(TSharedPtr<FNeuralMorphBoneGroupsTreeElement> InItem, const TSharedRef<STableViewBase>& OwnerTable);
+		TSharedPtr<SWidget> CreateContextMenuWidget() const;
+
+		void UpdateTreeElements();
+		void RefreshTree(bool bBroadcastPropertyChanged);
+		bool BroadcastModelPropertyChanged(const FName PropertyName);
+		TSharedPtr<SWidget> CreateContextWidget() const;
+
+		void OnCreateBoneGroup();
+		void OnDeleteSelectedItems();
+		void OnClearBoneGroups();
+		void OnAddBoneToGroup();
+
+	private:
+		TArray<TSharedPtr<FNeuralMorphBoneGroupsTreeElement>> RootElements;
+		TSharedPtr<UE::MLDeformer::SMLDeformerInputWidget> InputWidget;
+		FNeuralMorphEditorModel* EditorModel = nullptr;
+		FText SectionTitle;
+	};
+
+}	// namespace UE::NeuralMorphModel
diff --git a/Engine/Plugins/Animation/MLDeformer/NeuralMorphModel/Source/NeuralMorphModelEditor/Public/SNeuralMorphCurveGroupsWidget.h b/Engine/Plugins/Animation/MLDeformer/NeuralMorphModel/Source/NeuralMorphModelEditor/Public/SNeuralMorphCurveGroupsWidget.h
new file mode 100644
--- /dev/null
+++ b/Engine/Plugins/Animation/MLDeformer/NeuralMorphModel/Source/NeuralMorphModelEditor/Public/SNeuralMorphCurveGroupsWidget.h
@@ -0,0 +1,114 @@
+// Copyright Epic Games, Inc. All Rights Reserved.
+
+#pragma once
+
+#include "CoreTypes.h"
+#include "Templates/SharedPointer.h"
+#include "Widgets/SCompoundWidget.h"
+#include "Widgets/Views/STreeView.h"
+#include "Framework/Commands/Commands.h"
+#include "MLDeformerEditorStyle.h"
+
+class SWidget;
+
+namespace UE::MLDeformer
+{
+	class SMLDeformerInputWidget;
+};
+
+namespace UE::NeuralMorphModel
+{
+	class FNeuralMorphEditorModel;
+	class SNeuralMorphCurveGroupsWidget;
+
+	class FNeuralMorphCurveGroupsCommands
+		: public TCommands<FNeuralMorphCurveGroupsCommands>
+	{
+	public:
+		FNeuralMorphCurveGroupsCommands();
+		virtual void RegisterCommands() override;
+
+	public:
+		TSharedPtr<FUICommandInfo> CreateGroup;
+		TSharedPtr<FUICommandInfo> DeleteSelectedItems;
+		TSharedPtr<FUICommandInfo> ClearGroups;
+		TSharedPtr<FUICommandInfo> AddCurveToGroup;
+	};
+
+
+	class FNeuralMorphCurveGroupsTreeElement
+		: public TSharedFromThis<FNeuralMorphCurveGroupsTreeElement>
+	{
+	public:
+		TSharedRef<ITableRow> MakeTreeRowWidget(const TSharedRef<STableViewBase>& InOwnerTable, TSharedRef<FNeuralMorphCurveGroupsTreeElement> InTreeElement, TSharedPtr<SNeuralMorphCurveGroupsWidget> InTreeWidget);
+
+		bool IsGroup() const { return GroupIndex != INDEX_NONE; }
+
+	public:
+		FName Name;
+		TArray<TSharedPtr<FNeuralMorphCurveGroupsTreeElement>> Children;
+		TWeakPtr<FNeuralMorphCurveGroupsTreeElement> ParentGroup;
+		FSlateColor TextColor;
+		int32 GroupIndex = INDEX_NONE;
+		int32 GroupCurveIndex = INDEX_NONE;
+	};
+
+
+	class SNeuralMorphCurveGroupsTreeRowWidget 
+		: public STableRow<TSharedPtr<FNeuralMorphCurveGroupsTreeElement>>
+	{
+	public:
+		void Construct(const FArguments& InArgs, const TSharedRef<STableViewBase>& OwnerTable, TSharedRef<FNeuralMorphCurveGroupsTreeElement> InTreeElement, TSharedPtr<SNeuralMorphCurveGroupsWidget> InTreeView);
+
+	private:
+		TWeakPtr<FNeuralMorphCurveGroupsTreeElement> WeakTreeElement;
+		FText GetName() const;
+
+		friend class SNeuralMorphCurveGroupsWidget; 
+	};
+
+
+	class NEURALMORPHMODELEDITOR_API SNeuralMorphCurveGroupsWidget
+		: public STreeView<TSharedPtr<FNeuralMorphCurveGroupsTreeElement>>
+	{
+	public:
+		SLATE_BEGIN_ARGS(SNeuralMorphCurveGroupsWidget) {}
+		SLATE_ARGUMENT(FNeuralMorphEditorModel*, EditorModel)
+		SLATE_ARGUMENT(TSharedPtr<UE::MLDeformer::SMLDeformerInputWidget>, InputWidget)
+		SLATE_END_ARGS()
+
+		void Construct(const FArguments& InArgs);
+		void BindCommands(TSharedPtr<FUICommandList> CommandList);
+		void Refresh();
+
+		FText GetSectionTitle() const;
+		int32 GetNumSelectedGroups() const;
+		TSharedPtr<UE::MLDeformer::SMLDeformerInputWidget> GetInputWidget() const;
+
+	private:
+		void AddElement(TSharedPtr<FNeuralMorphCurveGroupsTreeElement> Element, TSharedPtr<FNeuralMorphCurveGroupsTreeElement> ParentElement);	
+		void HandleGetChildrenForTree(TSharedPtr<FNeuralMorphCurveGroupsTreeElement> InItem, TArray<TSharedPtr<FNeuralMorphCurveGroupsTreeElement>>& OutChildren);	
+		void OnSelectionChanged(TSharedPtr<FNeuralMorphCurveGroupsTreeElement> Selection, ESelectInfo::Type SelectInfo);
+		const TArray<TSharedPtr<FNeuralMorphCurveGroupsTreeElement>>& GetRootElements() const;
+		FReply OnKeyDown(const FGeometry& InGeometry, const FKeyEvent& InKeyEvent) override;
+		TSharedRef<ITableRow> MakeTableRowWidget(TSharedPtr<FNeuralMorphCurveGroupsTreeElement> InItem, const TSharedRef<STableViewBase>& OwnerTable);
+		TSharedPtr<SWidget> CreateContextMenuWidget() const;
+
+		void UpdateTreeElements();
+		void RefreshTree(bool bBroadcastPropertyChanged);
+		bool BroadcastModelPropertyChanged(const FName PropertyName);
+		TSharedPtr<SWidget> CreateContextWidget() const;
+
+		void OnCreateCurveGroup();
+		void OnDeleteSelectedItems();
+		void OnClearCurveGroups();
+		void OnAddCurveToGroup();
+
+	private:
+		TArray<TSharedPtr<FNeuralMorphCurveGroupsTreeElement>> RootElements;
+		TSharedPtr<UE::MLDeformer::SMLDeformerInputWidget> InputWidget;
+		FNeuralMorphEditorModel* EditorModel = nullptr;
+		FText SectionTitle;
+	};
+
+}	// namespace UE::NeuralMorphModel
diff --git a/Engine/Plugins/Animation/MLDeformer/NeuralMorphModel/Source/NeuralMorphModelEditor/Public/SNeuralMorphInputWidget.h b/Engine/Plugins/Animation/MLDeformer/NeuralMorphModel/Source/NeuralMorphModelEditor/Public/SNeuralMorphInputWidget.h
new file mode 100644
--- /dev/null
+++ b/Engine/Plugins/Animation/MLDeformer/NeuralMorphModel/Source/NeuralMorphModelEditor/Public/SNeuralMorphInputWidget.h
@@ -0,0 +1,103 @@
+// Copyright Epic Games, Inc. All Rights Reserved.
+
+#pragma once
+
+#include "CoreTypes.h"
+#include "Templates/SharedPointer.h"
+#include "SMLDeformerInputWidget.h"
+#include "Framework/Commands/Commands.h"
+
+class FMenuBuilder;
+
+namespace UE::MLDeformer
+{
+	class FMLDeformerEditorModel;
+};
+
+namespace UE::NeuralMorphModel
+{
+	class SNeuralMorphBoneGroupsWidget;
+	class SNeuralMorphCurveGroupsWidget;
+	class FNeuralMorphBoneGroupsTreeElement;
+	class FNeuralMorphCurveGroupsTreeElement;
+
+	class NEURALMORPHMODELEDITOR_API FNeuralMorphInputWidgetCommands
+		: public TCommands<FNeuralMorphInputWidgetCommands>
+	{
+	public:
+		FNeuralMorphInputWidgetCommands();
+		virtual void RegisterCommands() override;
+
+	public:
+		TSharedPtr<FUICommandInfo> ResetAllBoneMasks;
+		TSharedPtr<FUICommandInfo> ResetSelectedBoneMasks;
+		TSharedPtr<FUICommandInfo> ExpandSelectedBoneMasks;
+
+		TSharedPtr<FUICommandInfo> ResetAllBoneGroupMasks;
+		TSharedPtr<FUICommandInfo> ResetSelectedBoneGroupMasks;
+		TSharedPtr<FUICommandInfo> ExpandSelectedBoneGroupMasks;
+	};
+
+
+	class NEURALMORPHMODELEDITOR_API SNeuralMorphInputWidget
+		: public UE::MLDeformer::SMLDeformerInputWidget
+	{
+	public:
+		SLATE_BEGIN_ARGS(SNeuralMorphInputWidget) {}
+		SLATE_ARGUMENT(UE::MLDeformer::FMLDeformerEditorModel*, EditorModel)
+		SLATE_END_ARGS()
+
+		void Construct(const FArguments& InArgs);
+
+		// SMLDeformerInputWidget overrides.
+		virtual void Refresh() override;
+		virtual void AddInputBonesMenuItems(FMenuBuilder& MenuBuilder) override;
+		virtual void AddInputBonesPlusIconMenuItems(FMenuBuilder& MenuBuilder) override;
+		virtual void OnAddInputBones(const TArray<FName>& Names) override;
+		virtual void OnAddInputCurves(const TArray<FName>& Names) override;
+		virtual void OnAddAnimatedBones() override;
+		virtual void OnClearInputBones() override;
+		virtual void OnDeleteInputBones(const TArray<FName>& Names) override;
+		virtual void OnDeleteInputCurves(const TArray<FName>& Names) override;
+		virtual void OnSelectInputBone(FName BoneName) override;
+		virtual void OnSelectInputCurve(FName BoneName) override;
+		virtual void OnAddAnimatedCurves() override;
+		virtual void ClearSelectionForAllWidgetsExceptThis(TSharedPtr<SWidget> ExceptThisWidget) override;
+		virtual TSharedPtr<SWidget> GetExtraBonePickerWidget() override;
+		// ~END SMLDeformerInputWidget overrides.
+
+		void AddInputBoneGroupsMenuItems(FMenuBuilder& MenuBuilder);
+		void AddInputBoneGroupsPlusIconMenuItems(FMenuBuilder& MenuBuilder);
+
+		void OnSelectInputBoneGroup(TSharedPtr<FNeuralMorphBoneGroupsTreeElement> Element);
+		void OnSelectInputCurveGroup(TSharedPtr<FNeuralMorphCurveGroupsTreeElement> Element);
+
+		TSharedPtr<FUICommandList> GetBoneGroupsCommandList() const	{ return BoneGroupsCommandList; }
+		TSharedPtr<FUICommandList> GetCurveGroupsCommandList() const { return CurveGroupsCommandList; }
+
+		int32 GetHierarchyDepth() const { return HierarchyDepth; }
+
+	protected:
+		void CreateBoneGroupsSection();
+		void CreateCurveGroupsSection();
+
+		void ResetAllBoneMasks();
+		void ResetSelectedBoneMasks();
+		void ExpandBoneMasks();
+
+		void ResetAllBoneGroupMasks();
+		void ResetSelectedBoneGroupMasks();
+		void ExpandBoneGroupMasks();
+
+		void BindCommands();
+		FReply ShowBoneGroupsManageContextMenu();
+		FReply ShowCurveGroupsManageContextMenu();
+
+	protected:
+		TSharedPtr<SNeuralMorphBoneGroupsWidget> BoneGroupsWidget;
+		TSharedPtr<SNeuralMorphCurveGroupsWidget> CurveGroupsWidget;
+		TSharedPtr<FUICommandList> BoneGroupsCommandList;
+		TSharedPtr<FUICommandList> CurveGroupsCommandList;
+		int32 HierarchyDepth = 2;
+	};
+}	// namespace UE::NeuralMorphModel
