[MLDeformer] Improved reconstruction quality by the following changes:
- Changed to a smoothed L1 loss function, adding a new smoothing loss beta property to the training settings.
- Changed to cosine annealing learning rate scheduler, removing the learning rate decay setting.
- Removed some check which would fail when switching from local to global mode if you trained using local before. This check wasn't correct.

diff --git a/Engine/Plugins/Animation/MLDeformer/NeuralMorphModel/Source/NeuralMorphModel/Private/NeuralMorphModelInstance.cpp b/Engine/Plugins/Animation/MLDeformer/NeuralMorphModel/Source/NeuralMorphModel/Private/NeuralMorphModelInstance.cpp
--- a/Engine/Plugins/Animation/MLDeformer/NeuralMorphModel/Source/NeuralMorphModel/Private/NeuralMorphModelInstance.cpp
+++ b/Engine/Plugins/Animation/MLDeformer/NeuralMorphModel/Source/NeuralMorphModel/Private/NeuralMorphModelInstance.cpp
@@ -60,9 +60,6 @@ void UNeuralMorphModelInstance::FillNetworkInputs()
 	const UNeuralMorphMLP* GroupMLP = MorphNetwork->GetGroupMLP();
 	if (GroupMLP)
 	{
-		// Only the local mode should be using the group MLP.
-		check(Cast<UNeuralMorphModel>(Model)->GetModelMode() == ENeuralMorphMode::Local);
-
 		// Write the bone transforms.
 		int32 Offset = 0;
 		float* GroupNetworkInputs = NetworkInstance->GetGroupInputs().GetData();
diff --git a/Engine/Plugins/Animation/MLDeformer/NeuralMorphModel/Source/NeuralMorphModel/Public/NeuralMorphModel.h b/Engine/Plugins/Animation/MLDeformer/NeuralMorphModel/Source/NeuralMorphModel/Public/NeuralMorphModel.h
--- a/Engine/Plugins/Animation/MLDeformer/NeuralMorphModel/Source/NeuralMorphModel/Public/NeuralMorphModel.h
+++ b/Engine/Plugins/Animation/MLDeformer/NeuralMorphModel/Source/NeuralMorphModel/Public/NeuralMorphModel.h
@@ -72,7 +72,7 @@ public:
 	int32 GetNumIterations() const									{ return NumIterations; }
 	int32 GetBatchSize() const										{ return BatchSize; }
 	float GetLearningRate() const									{ return LearningRate; }
-	float GetLearningRateDecay() const								{ return LearningRateDecay; }
+	float GetSmoothLossBeta() const									{ return SmoothLossBeta; }
 	float GetRegularizationFactor() const							{ return RegularizationFactor; }
 	bool IsBoneMaskingEnabled() const								{ return bEnableBoneMasks; }
 
@@ -177,15 +177,6 @@ public:
 	UPROPERTY(EditAnywhere, BlueprintReadWrite, AdvancedDisplay, Category = "Training Settings", meta = (ClampMin = "0.000001", ClampMax = "1.0"))
 	float LearningRate = 0.001f;
 
-	/** 
-	 * The learning rate decay rate. If this is set to 1, the learning rate will not change. 
-	 * This is a multiplication factor. Every 1000 iterations, the new learning rate will be equal to
-	 * CurrentLearningRate * LearningRateDecay. This allows us to take slightly larger steps in the first iterations, and slowly take smaller steps.
-	 * Generally you want this to be between 0.9 and 1.
-	 */
-	UPROPERTY(EditAnywhere, BlueprintReadWrite, AdvancedDisplay, Category = "Training Settings", meta = (ClampMin = "0.1", ClampMax = "1.0"))
-	float LearningRateDecay = 0.98f;
-
 	/** 
 	 * The regularization factor. Higher values can help generate more sparse morph targets, but can also lead to visual artifacts. 
 	 * A value of 0 disables the regularization, and gives the highest quality, at the cost of higher runtime memory usage.
@@ -202,6 +193,15 @@ public:
 	UPROPERTY(EditAnywhere, BlueprintReadWrite, AdvancedDisplay, Category = "Training Settings", meta = (EditCondition = "Mode == ENeuralMorphMode::Local"))
 	bool bEnableBoneMasks = false;
 
+	/**
+	 * The beta parameter in the smooth L1 loss function, which describes below which absolute error to use a squared term. If the error is above or equal to this beta value, it will use the L1 loss.
+	 * This is a non-negative value, where 0 makes it behave exactly the same as an L1 loss.
+	 * The value entered represents how many cm of error to allow before an L1 loss is used. Typically higher values give smoother results.
+	 * If you see some noise in the trained results, even with large amount of samples and iterations, try increasing this value.
+	 */
+	UPROPERTY(EditAnywhere, BlueprintReadWrite, AdvancedDisplay, Category = "Training Settings", meta = (ClampMin = "0.0", ClampMax = "10.0"))
+	float SmoothLossBeta = 1.0f;
+
 	/**
 	 * The neural morph model network.
 	 * This is a specialized network used for inference, for performance reasons.
diff --git a/Engine/Plugins/Animation/MLDeformer/NeuralMorphModel/Source/NeuralMorphModelEditor/Private/NeuralMorphModelDetails.cpp b/Engine/Plugins/Animation/MLDeformer/NeuralMorphModel/Source/NeuralMorphModelEditor/Private/NeuralMorphModelDetails.cpp
--- a/Engine/Plugins/Animation/MLDeformer/NeuralMorphModel/Source/NeuralMorphModelEditor/Private/NeuralMorphModelDetails.cpp
+++ b/Engine/Plugins/Animation/MLDeformer/NeuralMorphModel/Source/NeuralMorphModelEditor/Private/NeuralMorphModelDetails.cpp
@@ -46,8 +46,8 @@ namespace UE::NeuralMorphModel
 		TrainingSettingsCategoryBuilder->AddProperty(GET_MEMBER_NAME_CHECKED(UNeuralMorphModel, NumIterations), UNeuralMorphModel::StaticClass());
 		TrainingSettingsCategoryBuilder->AddProperty(GET_MEMBER_NAME_CHECKED(UNeuralMorphModel, BatchSize), UNeuralMorphModel::StaticClass());
 		TrainingSettingsCategoryBuilder->AddProperty(GET_MEMBER_NAME_CHECKED(UNeuralMorphModel, LearningRate), UNeuralMorphModel::StaticClass());
-		TrainingSettingsCategoryBuilder->AddProperty(GET_MEMBER_NAME_CHECKED(UNeuralMorphModel, LearningRateDecay), UNeuralMorphModel::StaticClass());
 		TrainingSettingsCategoryBuilder->AddProperty(GET_MEMBER_NAME_CHECKED(UNeuralMorphModel, RegularizationFactor), UNeuralMorphModel::StaticClass());
+		TrainingSettingsCategoryBuilder->AddProperty(GET_MEMBER_NAME_CHECKED(UNeuralMorphModel, SmoothLossBeta), UNeuralMorphModel::StaticClass());
 		TrainingSettingsCategoryBuilder->AddProperty(GET_MEMBER_NAME_CHECKED(UNeuralMorphModel, bEnableBoneMasks), UNeuralMorphModel::StaticClass());
 	}
 }	// namespace UE::NeuralMorphModel
