[MLDeformer] Several small improvements, including:
- Extra Get functions to get private members.
- Verbose logging where non-verbose logging wasn't needed.
- Some additional null safety checks.
- Prevent sampler reinit when not needed.
- Models can specify the number of floats per bone.

diff --git a/Engine/Plugins/Animation/MLDeformer/MLDeformerFramework/Source/MLDeformerFramework/Private/MLDeformerGeomCacheHelpers.cpp b/Engine/Plugins/Animation/MLDeformer/MLDeformerFramework/Source/MLDeformerFramework/Private/MLDeformerGeomCacheHelpers.cpp
--- a/Engine/Plugins/Animation/MLDeformer/MLDeformerFramework/Source/MLDeformerFramework/Private/MLDeformerGeomCacheHelpers.cpp
+++ b/Engine/Plugins/Animation/MLDeformer/MLDeformerFramework/Source/MLDeformerFramework/Private/MLDeformerGeomCacheHelpers.cpp
@@ -220,7 +220,7 @@ namespace UE::MLDeformer
 					bFoundMatch = true;
 					if (!bSuppressLog)
 					{
-						UE_LOG(LogMLDeformer, Display, TEXT("Mapped geom cache track '%s' to SkelMesh submesh '%s'"), *Track->GetName(), *SkelMeshName);
+						UE_LOG(LogMLDeformer, Verbose, TEXT("Mapped geom cache track '%s' to SkelMesh submesh '%s'"), *Track->GetName(), *SkelMeshName);
 					}
 					break;
 				} // If the track name matches the skeletal meshes internal mesh name.
diff --git a/Engine/Plugins/Animation/MLDeformer/MLDeformerFramework/Source/MLDeformerFramework/Private/MLDeformerInputInfo.cpp b/Engine/Plugins/Animation/MLDeformer/MLDeformerFramework/Source/MLDeformerFramework/Private/MLDeformerInputInfo.cpp
--- a/Engine/Plugins/Animation/MLDeformer/MLDeformerFramework/Source/MLDeformerFramework/Private/MLDeformerInputInfo.cpp
+++ b/Engine/Plugins/Animation/MLDeformer/MLDeformerFramework/Source/MLDeformerFramework/Private/MLDeformerInputInfo.cpp
@@ -130,6 +130,11 @@ void UMLDeformerInputInfo::SetNumTargetVertices(int32 NumVerts)
 
 bool UMLDeformerInputInfo::IsCompatible(USkeletalMesh* InSkeletalMesh) const
 {
+	if (SkeletalMesh.IsNull())
+	{
+		return true;
+	}
+
 	if (InSkeletalMesh == nullptr)
 	{
 		return false;
diff --git a/Engine/Plugins/Animation/MLDeformer/MLDeformerFramework/Source/MLDeformerFramework/Public/MLDeformerGeomCacheModel.h b/Engine/Plugins/Animation/MLDeformer/MLDeformerFramework/Source/MLDeformerFramework/Public/MLDeformerGeomCacheModel.h
--- a/Engine/Plugins/Animation/MLDeformer/MLDeformerFramework/Source/MLDeformerFramework/Public/MLDeformerGeomCacheModel.h
+++ b/Engine/Plugins/Animation/MLDeformer/MLDeformerFramework/Source/MLDeformerFramework/Public/MLDeformerGeomCacheModel.h
@@ -61,6 +61,13 @@ public:
 	 */
 	UGeometryCache* GetGeometryCache()						{ return GeometryCache.LoadSynchronous(); }
 
+	/**
+	 * Set the geometry cache object to use for training.
+	 * Keep in mind that the editor still needs to handle a change of this property for things to be initialized correctly.
+	 * @param GeomCache The geometry cache to use for training.
+	 */
+	void SetGeometryCache(UGeometryCache* GeomCache)		{ GeometryCache = GeomCache; }
+
 	/**
 	 * Get the mapping between geometry cache tracks and meshes inside the skeletal mesh.
 	 * This lets us know what parts of the skeletal mesh are related to what geometry cache tracks.
diff --git a/Engine/Plugins/Animation/MLDeformer/MLDeformerFramework/Source/MLDeformerFramework/Public/MLDeformerModel.h b/Engine/Plugins/Animation/MLDeformer/MLDeformerFramework/Source/MLDeformerFramework/Public/MLDeformerModel.h
--- a/Engine/Plugins/Animation/MLDeformer/MLDeformerFramework/Source/MLDeformerFramework/Public/MLDeformerModel.h
+++ b/Engine/Plugins/Animation/MLDeformer/MLDeformerFramework/Source/MLDeformerFramework/Public/MLDeformerModel.h
@@ -175,6 +175,18 @@ public:
 	 */
 	virtual FString GetDefaultDeformerGraphAssetPath() const	{ return FString(); }
 
+	/** 
+	 * Get the number of floats used to represent a single bone rotation, used as input to the neural networks.
+	 * @return The number of floats per individual bone.
+	 */
+	virtual int32 GetNumFloatsPerBone() const				{ return 6; }
+
+	/** 
+	 * Get the number of floats used to represent a single bone rotation, used as input to the neural networks.
+	 * @return The number of floats per individual bone.
+	 */
+	virtual int32 GetNumFloatsPerCurve() const				{ return 1; }
+
 	/**
 	 * Get the skeletal mesh that is used during training.
 	 * You typically want to apply the ML Deformer on this specific skeletal mesh in your game as well.
@@ -402,6 +414,13 @@ public:
 	const UAnimSequence* GetAnimSequence() const				{ return AnimSequence.LoadSynchronous();  }
 	UAnimSequence* GetAnimSequence()							{ return AnimSequence.LoadSynchronous(); }
 
+	/**
+	 * Set the animation sequence object to use for training.
+	 * Keep in mind that the editor still needs to handle a change of this property for things to be initialized correctly.
+	 * @param AnimSeq The animation sequence to use for training.
+	 */
+	void SetAnimSequence(UAnimSequence* AnimSeq)				{ AnimSequence = AnimSeq; }
+
 	/**
 	 * Get the maximum number of training frames to use during training.
 	 * Your training anim sequence might contain say 10000 frames, but for quickly iterating you might
@@ -516,12 +535,6 @@ public:
 
 #endif	// #if WITH_EDITORONLY_DATA
 
-	/** @return The number of floats per bone in network input. */
-	virtual int32 GetNumFloatsPerBone() const { return NumFloatsPerBone; }
-
-	/** @return The number of floats per curve in network input. */
-	virtual int32 GetNumFloatsPerCurve() const { return NumFloatsPerCurve; }
-
 protected:
 	/**
 	 * Set the training input information.
@@ -611,9 +624,11 @@ private:
 	TObjectPtr<USkeletalMesh> SkeletalMesh = nullptr;
 
 	/** The number of floats per bone in network input. */
+	UE_DEPRECATED(5.3, "This will be removed")
 	static constexpr int32 NumFloatsPerBone = 6;
 
 	/** The number of floats per curve in network input. */
+	UE_DEPRECATED(5.3, "This will be removed")
 	static constexpr int32 NumFloatsPerCurve = 1;
 
 #if WITH_EDITORONLY_DATA
diff --git a/Engine/Plugins/Animation/MLDeformer/MLDeformerFramework/Source/MLDeformerFramework/Public/MLDeformerMorphModel.h b/Engine/Plugins/Animation/MLDeformer/MLDeformerFramework/Source/MLDeformerFramework/Public/MLDeformerMorphModel.h
--- a/Engine/Plugins/Animation/MLDeformer/MLDeformerFramework/Source/MLDeformerFramework/Public/MLDeformerMorphModel.h
+++ b/Engine/Plugins/Animation/MLDeformer/MLDeformerFramework/Source/MLDeformerFramework/Public/MLDeformerMorphModel.h
@@ -11,6 +11,11 @@ struct FExternalMorphSet;
 struct FExternalMorphSetWeights;
 class USkinnedMeshComponent;
 
+/**
+ * The morph model base class.
+ * This is the base class for models that generate and drive morph targets.
+ * Use this in combination with UMLDeformerMorphModelInstance or inherited classes.
+ */
 UCLASS()
 class MLDEFORMERFRAMEWORK_API UMLDeformerMorphModel
 	: public UMLDeformerGeomCacheModel
@@ -52,9 +57,12 @@ public:
 	void SetMorphDeltaZeroThreshold(float Threshold)				{ MorphDeltaZeroThreshold = Threshold; }
 	void SetMorphCompressionlevel(float Tolerance)					{ MorphCompressionLevel = Tolerance; }
 	void SetIncludeMorphTargetNormals(bool bInclude)				{ bIncludeNormals = bInclude; }
-	void SetWeightMask(EMLDeformerMaskChannel Channel)				{ MaskChannel = Channel; }
+	void SetMaskChannel(EMLDeformerMaskChannel Channel)				{ MaskChannel = Channel; }
 	void SetInvertMaskChannel(bool bInvert)							{ bInvertMaskChannel = bInvert; }
 
+	UE_DEPRECATED(5.3, "Use SetMaskChannel instead.")
+	void SetWeightMask(EMLDeformerMaskChannel Channel)				{ MaskChannel = Channel; }
+
 	UE_DEPRECATED(5.2, "Please use GetMorphDeltaZeroThreshold instead.")
 	float GetMorphTargetDeltaThreshold() const						{ return MorphDeltaZeroThreshold; }
 
diff --git a/Engine/Plugins/Animation/MLDeformer/MLDeformerFramework/Source/MLDeformerFrameworkEditor/Private/MLDeformerEditorModel.cpp b/Engine/Plugins/Animation/MLDeformer/MLDeformerFramework/Source/MLDeformerFrameworkEditor/Private/MLDeformerEditorModel.cpp
--- a/Engine/Plugins/Animation/MLDeformer/MLDeformerFramework/Source/MLDeformerFrameworkEditor/Private/MLDeformerEditorModel.cpp
+++ b/Engine/Plugins/Animation/MLDeformer/MLDeformerFramework/Source/MLDeformerFrameworkEditor/Private/MLDeformerEditorModel.cpp
@@ -1326,7 +1326,7 @@ namespace UE::MLDeformer
 		ClampCurrentTrainingFrameIndex();
 
 		// If we have no Persona toolkit yet, then it is not yet safe to init the sampler.
-		if (Editor->GetPersonaToolkitPointer() != nullptr)
+		if (Editor->GetPersonaToolkitPointer() != nullptr && !Sampler->IsInitialized())
 		{
 			Sampler->Init(this);
 		}
diff --git a/Engine/Plugins/Animation/MLDeformer/NeuralMorphModel/Source/NeuralMorphModel/Private/NeuralMorphModel.cpp b/Engine/Plugins/Animation/MLDeformer/NeuralMorphModel/Source/NeuralMorphModel/Private/NeuralMorphModel.cpp
--- a/Engine/Plugins/Animation/MLDeformer/NeuralMorphModel/Source/NeuralMorphModel/Private/NeuralMorphModel.cpp
+++ b/Engine/Plugins/Animation/MLDeformer/NeuralMorphModel/Source/NeuralMorphModel/Private/NeuralMorphModel.cpp
@@ -3,6 +3,7 @@
 #include "NeuralMorphModelVizSettings.h"
 #include "NeuralMorphModelInstance.h"
 #include "NeuralMorphInputInfo.h"
+#include "NeuralMorphNetwork.h"
 #include "MLDeformerAsset.h"
 #include "MLDeformerComponent.h"
 #include "UObject/UObjectGlobals.h"
@@ -117,4 +118,13 @@ void UNeuralMorphModel::SetNeuralMorphNetwork(UNeuralMorphNetwork* Net)
 	GetReinitModelInstanceDelegate().Broadcast();
 }
 
+int32 UNeuralMorphModel::GetNumFloatsPerCurve() const
+{
+	if (NeuralMorphNetwork)
+	{
+		return NeuralMorphNetwork->GetNumFloatsPerCurve();
+	}
+	return (Mode == ENeuralMorphMode::Local) ? 6 : 1;
+}
+
 #undef LOCTEXT_NAMESPACE
diff --git a/Engine/Plugins/Animation/MLDeformer/NeuralMorphModel/Source/NeuralMorphModel/Private/NeuralMorphModelInstance.cpp b/Engine/Plugins/Animation/MLDeformer/NeuralMorphModel/Source/NeuralMorphModel/Private/NeuralMorphModelInstance.cpp
--- a/Engine/Plugins/Animation/MLDeformer/NeuralMorphModel/Source/NeuralMorphModel/Private/NeuralMorphModelInstance.cpp
+++ b/Engine/Plugins/Animation/MLDeformer/NeuralMorphModel/Source/NeuralMorphModel/Private/NeuralMorphModelInstance.cpp
@@ -42,7 +42,7 @@ void UNeuralMorphModelInstance::FillNetworkInputs()
 	Super::SetBoneTransforms(MainInputs, NumMainMLPInputs, 0);
 
 	// Write the curve input values, start writing after the bone inputs.
-	const int32 CurveWriteOffset = Model->GetInputInfo()->GetNumBones() * 6;	// 6 Floats per bone.
+	const int32 CurveWriteOffset = Model->GetInputInfo()->GetNumBones() * Model->GetNumFloatsPerBone();
 	SetCurveValues(MainInputs, NumMainMLPInputs, CurveWriteOffset);
 
 	// Normalize the inputs.
@@ -60,6 +60,8 @@ void UNeuralMorphModelInstance::FillNetworkInputs()
 	const UNeuralMorphMLP* GroupMLP = MorphNetwork->GetGroupMLP();
 	if (GroupMLP)
 	{
+		const int32 NumFloatsPerBone = Model->GetNumFloatsPerBone();
+
 		// Write the bone transforms.
 		int32 Offset = 0;
 		float* GroupNetworkInputs = NetworkInstance->GetGroupInputs().GetData();
@@ -69,7 +71,7 @@ void UNeuralMorphModelInstance::FillNetworkInputs()
 			const int32 BoneIndex = BoneGroupIndices[Index];
 			if (BoneIndex != INDEX_NONE)
 			{
-				const int32 BoneOffset = BoneIndex * 6;
+				const int32 BoneOffset = BoneIndex * NumFloatsPerBone;
 				GroupNetworkInputs[Offset++] = MainInputs[BoneOffset];
 				GroupNetworkInputs[Offset++] = MainInputs[BoneOffset + 1];
 				GroupNetworkInputs[Offset++] = MainInputs[BoneOffset + 2];
@@ -88,9 +90,8 @@ void UNeuralMorphModelInstance::FillNetworkInputs()
 			}
 		}
 
-		// Write the curve values.
+		// Write the curve values for all curve groups.
 		const int32 NumFloatsPerCurve = MorphNetwork->GetNumFloatsPerCurve();
-		check(MorphNetwork->GetNumFloatsPerCurve() == 6);	// Local mode curve values are expected to use 6 floats.
 		const int32 NumCurveGroups = CurveGroupIndices.Num();
 		for (int32 Index = 0; Index < CurveGroupIndices.Num(); ++Index)
 		{
@@ -118,6 +119,7 @@ void UNeuralMorphModelInstance::FillNetworkInputs()
 	}
 }
 
+// This will fill the curve values for the curves, with curve groups excluded.
 int64 UNeuralMorphModelInstance::SetCurveValues(float* OutputBuffer, int64 OutputBufferSize, int64 StartIndex)
 {
 	UNeuralMorphModel* MorphModel = Cast<UNeuralMorphModel>(Model);
@@ -127,7 +129,7 @@ int64 UNeuralMorphModelInstance::SetCurveValues(float* OutputBuffer, int64 Outpu
 
 	int64 Index = StartIndex;
 	const int32 AssetNumCurves = InputInfo->GetNumCurves();
-	const int32 NumFloatsPerCurve = MorphNetwork ? MorphNetwork->GetNumFloatsPerCurve() : 1;
+	const int32 NumFloatsPerCurve = MorphNetwork->GetNumFloatsPerCurve();
 	const int32 NumCurveFloats = AssetNumCurves * NumFloatsPerCurve;
 	checkf((Index + NumCurveFloats) <= OutputBufferSize, TEXT("Writing curves past the end of the input buffer"));
 
diff --git a/Engine/Plugins/Animation/MLDeformer/NeuralMorphModel/Source/NeuralMorphModel/Public/NeuralMorphModel.h b/Engine/Plugins/Animation/MLDeformer/NeuralMorphModel/Source/NeuralMorphModel/Public/NeuralMorphModel.h
--- a/Engine/Plugins/Animation/MLDeformer/NeuralMorphModel/Source/NeuralMorphModel/Public/NeuralMorphModel.h
+++ b/Engine/Plugins/Animation/MLDeformer/NeuralMorphModel/Source/NeuralMorphModel/Public/NeuralMorphModel.h
@@ -57,6 +57,7 @@ public:
 	virtual FString GetDisplayName() const override					{ return "Neural Morph Model"; }
 	virtual UMLDeformerModelInstance* CreateModelInstance(UMLDeformerComponent* Component) override;
 	virtual UMLDeformerInputInfo* CreateInputInfo() override;
+	virtual int32 GetNumFloatsPerCurve() const override;
 	// ~END UMLDeformerModel overrides.
 
 	const TArray<FNeuralMorphBoneGroup>& GetBoneGroups() const		{ return BoneGroups; }
@@ -71,6 +72,7 @@ public:
 	int32 GetGlobalNumNeuronsPerLayer() const						{ return GlobalNumNeuronsPerLayer; }
 	int32 GetNumIterations() const									{ return NumIterations; }
 	int32 GetBatchSize() const										{ return BatchSize; }
+	int32 GetLocalNumMorphsPerBoneOrCurve() const					{ return LocalNumMorphTargetsPerBone; }
 	float GetLearningRate() const									{ return LearningRate; }
 	float GetSmoothLossBeta() const									{ return SmoothLossBeta; }
 	float GetRegularizationFactor() const							{ return RegularizationFactor; }
@@ -79,6 +81,8 @@ public:
 	void UpdateMissingGroupNames();
 	void SetNeuralMorphNetwork(UNeuralMorphNetwork* Net);
 
+	void SetNumIterations(int32 InNumIterations)					{ check(InNumIterations > 0); NumIterations = InNumIterations; }
+
 public:
 	/**
 	 * The set of bones that are grouped together and generate morph targets together as a whole.
