MLDeformer: Updated NeuralMorphModel to use new NNERuntimeBasicCpu plugin

diff --git a/Engine/Plugins/Animation/MLDeformer/NeuralMorphModel/NeuralMorphModel.uplugin b/Engine/Plugins/Animation/MLDeformer/NeuralMorphModel/NeuralMorphModel.uplugin
--- a/Engine/Plugins/Animation/MLDeformer/NeuralMorphModel/NeuralMorphModel.uplugin
+++ b/Engine/Plugins/Animation/MLDeformer/NeuralMorphModel/NeuralMorphModel.uplugin
@@ -42,6 +42,10 @@
 		{
 			"Name": "MLDeformerFramework",
 			"Enabled": true
+		},
+		{
+			"Name": "NNERuntimeBasicCpu",
+			"Enabled": true
 		}
 	]
 }
diff --git a/Engine/Plugins/Animation/MLDeformer/NeuralMorphModel/Source/NeuralMorphModel/NeuralMorphModel.Build.cs b/Engine/Plugins/Animation/MLDeformer/NeuralMorphModel/Source/NeuralMorphModel/NeuralMorphModel.Build.cs
--- a/Engine/Plugins/Animation/MLDeformer/NeuralMorphModel/Source/NeuralMorphModel/NeuralMorphModel.Build.cs
+++ b/Engine/Plugins/Animation/MLDeformer/NeuralMorphModel/Source/NeuralMorphModel/NeuralMorphModel.Build.cs
@@ -42,7 +42,9 @@ public class NeuralMorphModel : ModuleRules
 				"Projects",
 				"RenderCore",
 				"RHI",
-				"MLDeformerFramework"
+				"MLDeformerFramework",
+				"NNE",
+				"NNERuntimeBasicCpu",
 			}
 		);
 		
diff --git a/Engine/Plugins/Animation/MLDeformer/NeuralMorphModel/Source/NeuralMorphModel/Private/NeuralMorphModelInstance.cpp b/Engine/Plugins/Animation/MLDeformer/NeuralMorphModel/Source/NeuralMorphModel/Private/NeuralMorphModelInstance.cpp
--- a/Engine/Plugins/Animation/MLDeformer/NeuralMorphModel/Source/NeuralMorphModel/Private/NeuralMorphModelInstance.cpp
+++ b/Engine/Plugins/Animation/MLDeformer/NeuralMorphModel/Source/NeuralMorphModel/Private/NeuralMorphModelInstance.cpp
@@ -37,13 +37,12 @@ void UNeuralMorphModelInstance::FillNetworkInputs()
 
 	// Write the input bone transforms.
 	float* MainInputs = NetworkInstance->GetInputs().GetData();
-	const UNeuralMorphMLP* MainMLP = MorphNetwork->GetMainMLP();
-	const int32 NumMainMLPInputs = MainMLP->GetNumInputs();
-	Super::SetBoneTransforms(MainInputs, NumMainMLPInputs, 0);
+	const int32 NumMainNeuralNetworkInputs = MorphNetwork->GetNumMainInputs();
+	Super::SetBoneTransforms(MainInputs, NumMainNeuralNetworkInputs, 0);
 
 	// Write the curve input values, start writing after the bone inputs.
 	const int32 CurveWriteOffset = Model->GetInputInfo()->GetNumBones() * Model->GetNumFloatsPerBone();
-	SetCurveValues(MainInputs, NumMainMLPInputs, CurveWriteOffset);
+	SetCurveValues(MainInputs, NumMainNeuralNetworkInputs, CurveWriteOffset);
 
 	// Normalize the inputs.
 	const float* Means = MorphNetwork->GetInputMeans().GetData();
@@ -54,11 +53,10 @@ void UNeuralMorphModelInstance::FillNetworkInputs()
 		MainInputs[Index] = (MainInputs[Index] - Means[Index]) / Stds[Index];
 	}
 
-	// Write the input transforms for the group MLP.
+	// Write the input transforms for the group Neural Network.
 	// We reuse the same input values from the main network, and those values just got normalized.
 	// So there is no need to normalize them again.
-	const UNeuralMorphMLP* GroupMLP = MorphNetwork->GetGroupMLP();
-	if (GroupMLP)
+	if (MorphNetwork->GetGroupModel())
 	{
 		const int32 NumFloatsPerBone = Model->GetNumFloatsPerBone();
 
@@ -187,7 +185,7 @@ bool UNeuralMorphModelInstance::SetupInputs()
 	// Make sure we have a valid network.
 	UNeuralMorphModel* MorphModel = Cast<UNeuralMorphModel>(Model);
 	UNeuralMorphNetwork* MorphNetwork = MorphModel->GetNeuralMorphNetwork();
-	if (MorphNetwork == nullptr || MorphNetwork->GetMainMLP() == nullptr)
+	if (MorphNetwork == nullptr || MorphNetwork->GetMainModel() == nullptr)
 	{		
 		return false;
 	}
diff --git a/Engine/Plugins/Animation/MLDeformer/NeuralMorphModel/Source/NeuralMorphModel/Private/NeuralMorphNetwork.cpp b/Engine/Plugins/Animation/MLDeformer/NeuralMorphModel/Source/NeuralMorphModel/Private/NeuralMorphNetwork.cpp
--- a/Engine/Plugins/Animation/MLDeformer/NeuralMorphModel/Source/NeuralMorphModel/Private/NeuralMorphNetwork.cpp
+++ b/Engine/Plugins/Animation/MLDeformer/NeuralMorphModel/Source/NeuralMorphModel/Private/NeuralMorphNetwork.cpp
@@ -5,14 +5,20 @@
 #include "UObject/SoftObjectPath.h"
 #include "HAL/FileManagerGeneric.h"
 #include "Math/UnrealMathUtility.h"
+#include "Misc/FileHelper.h"
 
-#if NEURALMORPHMODEL_USE_ISPC
-	#include "NeuralMorphNetwork.ispc.generated.h"
-#endif
+#include "NNE.h"
+#include "NNERuntime.h"
+#include "NNERuntimeCPU.h"
+#include "NNEModelData.h"
+#include "NNERuntimeBasicCpu.h"
 
 //--------------------------------------------------------------------------
 // UNeuralMorphMLP
 //--------------------------------------------------------------------------
+
+PRAGMA_DISABLE_DEPRECATION_WARNINGS
+
 void UNeuralMorphMLP::Empty()
 {
 	Layers.Empty();
@@ -161,7 +167,7 @@ bool UNeuralMorphMLP::Load(FArchive& FileReader)
 		UE_LOG(LogNeuralMorphModel, Verbose, TEXT("MLP Layer %d --> NumWeights=%d (%dx%dx%d)   NumBiases=%d"), LayerIndex, Layer->Weights.Num(), Layer->NumInputs, Layer->NumOutputs, Layer->Depth, Layer->Biases.Num());
 	}
 
-	UE_LOG(LogNeuralMorphModel, Display, TEXT("Successfullly loaded neural morph MLP from file '%s'"), *FileReader.GetArchiveName());
+	UE_LOG(LogNeuralMorphModel, Display, TEXT("Successfully loaded neural morph MLP from file '%s'"), *FileReader.GetArchiveName());
 	return true;
 }
 
@@ -175,194 +181,371 @@ UNeuralMorphMLPLayer& UNeuralMorphMLP::GetLayer(int32 Index) const
 	return *Layers[Index].Get();
 }
 
+PRAGMA_ENABLE_DEPRECATION_WARNINGS
+
 //--------------------------------------------------------------------------
 // UNeuralMorphNetwork
 //--------------------------------------------------------------------------
-void UNeuralMorphNetwork::Empty()
+
+namespace UE::NeuralMorphModel::Private
 {
-	InputMeans.Empty();
-	InputStd.Empty();
+	// Magic and version number used to identify this file data transfered between
+	// python and the NeuralMorphModel. This should match what is in "nmm_shared.py"
+	static constexpr uint32 MagicNumber = 0x234A1304;
+	static constexpr uint32 VersionNumber = 1;
 
-	if (MainMLP)
+	namespace Serialization
 	{
-		MainMLP->ConditionalBeginDestroy();
-		MainMLP = nullptr;
-	}
+		//--------------------------------------------------------------------------
 
-	if (GroupMLP)
-	{
-		GroupMLP->ConditionalBeginDestroy();
-		GroupMLP = nullptr;
+		static inline void Align(uint64& InOutOffset, uint32 Alignment)
+		{
+			InOutOffset = ((InOutOffset + Alignment - 1) / Alignment) * Alignment;
+		}
+
+		//--------------------------------------------------------------------------
+
+		static inline void Size(uint64& InOutOffset, const uint32& In)
+		{
+			Align(InOutOffset, sizeof(uint32));
+			InOutOffset += sizeof(uint32);
+		}
+
+		static inline void Size(uint64& InOutOffset, const float& In)
+		{
+			Align(InOutOffset, alignof(float));
+			InOutOffset += sizeof(float);
+		}
+
+		static inline void Size(uint64& InOutOffset, const TConstArrayView<float> In)
+		{
+			Align(InOutOffset, 64);
+			InOutOffset += In.Num() * sizeof(float);
+		}
+
+		static inline void Size(uint64& InOutOffset, const TConstArrayView<uint16> In)
+		{
+			Align(InOutOffset, 64);
+			InOutOffset += In.Num() * sizeof(uint16);
+		}
+
+		//--------------------------------------------------------------------------
+
+		static inline void Load(uint64& InOutOffset, uint32& Out, TConstArrayView<uint8> Data)
+		{
+			Align(InOutOffset, sizeof(uint32));
+			Out = *((uint32*)&Data[InOutOffset]);
+			InOutOffset += sizeof(uint32);
+		}
+
+		static inline void Load(uint64& InOutOffset, int32& Out, TConstArrayView<uint8> Data)
+		{
+			Align(InOutOffset, sizeof(int32));
+			Out = *((int32*)&Data[InOutOffset]);
+			InOutOffset += sizeof(int32);
+		}
+
+		static inline void Load(uint64& InOutOffset, float& Out, TConstArrayView<uint8> Data)
+		{
+			Align(InOutOffset, sizeof(float));
+			Out = *((float*)&Data[InOutOffset]);
+			InOutOffset += sizeof(float);
+		}
+
+		static inline void Load(uint64& InOutOffset, TArrayView<float> Out, TConstArrayView<uint8> Data)
+		{
+			Align(InOutOffset, 64);
+			FMemory::Memcpy(Out.GetData(), &Data[InOutOffset], Out.Num() * sizeof(float));
+			Out = MakeArrayView<float>((float*)&Data[InOutOffset], Out.Num());
+			InOutOffset += Out.Num() * sizeof(float);
+		}
+
+		static inline void Load(uint64& InOutOffset, FString& Out, TConstArrayView<uint8> Data)
+		{
+			uint32 Num;
+			Load(InOutOffset, Num, Data);
+			Out = UTF8_TO_TCHAR(&Data[InOutOffset]);
+			InOutOffset += Num;
+		}
 	}
 
-	Mode = ENeuralMorphMode::Global;
-	NumMorphsPerBone = 0;
-	NumBones = 0;
-	NumCurves = 0;
-	NumGroups = 0;
-	NumItemsPerGroup = 0;
-	NumFloatsPerCurve = 1;
-}
+	//--------------------------------------------------------------------------
 
-bool UNeuralMorphNetwork::IsEmpty() const
-{
-	return !MainMLP;
-}
+	PRAGMA_DISABLE_DEPRECATION_WARNINGS
 
-UNeuralMorphMLP* UNeuralMorphNetwork::LoadMLP(FArchive& FileReader)
-{
-	UNeuralMorphMLP* MLP = NewObject<UNeuralMorphMLP>(this);
-	if (!MLP->Load(FileReader))
+	static inline void ConvertMLPToFileData(TArray<uint8>& OutBytes, TObjectPtr<UNeuralMorphMLP>& MLP)
 	{
+		UE::NNE::RuntimeBasic::FSequentialModelBuilder Builder;
+
+		const uint32 LayerNum = MLP->GetNumLayers();
+
+		for (uint32 LayerIdx = 0; LayerIdx < LayerNum; LayerIdx++)
+		{
+			const uint32 InputSize = MLP->GetLayer(LayerIdx).NumInputs;
+			const uint32 OutputSize = MLP->GetLayer(LayerIdx).NumOutputs;
+			const uint32 BlockNum = MLP->GetLayer(LayerIdx).Depth;
+			const TConstArrayView<float> Weights = MLP->GetLayer(LayerIdx).Weights;
+			const TConstArrayView<float> Biases = MLP->GetLayer(LayerIdx).Biases;
+
+			if (BlockNum == 1)
+			{
+				Builder.AddLinear(InputSize, OutputSize, Weights, Biases);
+			}
+			else
+			{
+				Builder.AddMultiLinear(InputSize, OutputSize, BlockNum, Weights, Biases);
+			}
+
+			Builder.AddELU();
+		}
+
 		MLP->ConditionalBeginDestroy();
-		MLP = nullptr;		
+		MLP = nullptr;
+
+		OutBytes.SetNumUninitialized(Builder.GetWriteByteNum());
+		Builder.WriteAndReset(OutBytes);
 	}
-	return MLP;
+
+	PRAGMA_ENABLE_DEPRECATION_WARNINGS
 }
 
-bool UNeuralMorphNetwork::Load(const FString& Filename)
+UNeuralMorphNetwork::UNeuralMorphNetwork() : Super() {}
+UNeuralMorphNetwork::UNeuralMorphNetwork(const FObjectInitializer& ObjectInitializer) : Super(ObjectInitializer) {}
+UNeuralMorphNetwork::UNeuralMorphNetwork(FVTableHelper& Helper) : Super(Helper) {}
+UNeuralMorphNetwork::~UNeuralMorphNetwork() {}
+
+PRAGMA_DISABLE_DEPRECATION_WARNINGS
+
+void UNeuralMorphNetwork::PostLoad()
 {
-	UE_LOG(LogNeuralMorphModel, Display, TEXT("Loading Neural Morph Network from file '%s'"), *Filename);
-	Empty();
+	Super::PostLoad();
 
-	// Create the file reader.
-	TUniquePtr<FArchive> FileReader(IFileManager::Get().CreateFileReader(*Filename));
-	if (!FileReader.IsValid())
+	if (MainMLP_DEPRECATED)
 	{
-		UE_LOG(LogNeuralMorphModel, Error, TEXT("Failed to create file reader to read file '%s'!"), *Filename);
-		Empty();
-		return false;
-	}
+		// If we are in Global mode and loading the legacy network format we need to Load NumMorphs 
+		// from old network data output size since it was not stored directly in the object before.
+		if (Mode == ENeuralMorphMode::Global)
+		{
+			NumMorphs = MainMLP_DEPRECATED->GetNumOutputs();
+		}
 
-	// Read the FOURCC, to identify the file type.
-	char FOURCC[4] {' ', ' ', ' ', ' '};
-	FileReader->Serialize(FOURCC, 4);
-	if (FileReader->IsError())
-	{
-		FileReader->Close();
-		UE_LOG(LogNeuralMorphModel, Error, TEXT("Failed to read the FOURCC!"));
-		Empty();
-		return false;
-	}
+		TArray<uint8> FileData;
+		UE::NeuralMorphModel::Private::ConvertMLPToFileData(FileData, MainMLP_DEPRECATED);
 
-	if (FOURCC[0] != 'N' || FOURCC[1] != 'M' || FOURCC[2] != 'M' || FOURCC[3] != 'N')	// NMMN (Neural Morph Model Network)
-	{
-		FileReader->Close();
-		UE_LOG(LogNeuralMorphModel, Error, TEXT("The file is not a valid valid neural morph network file type!"));
-		Empty();
-		return false;
+		if (!MainModelData)
+		{
+			MainModelData = NewObject<UNNEModelData>(this);
+		}
+
+		MainModelData->Init(TEXT("ubnne"), FileData);
+
+		FileData.Empty();
 	}
 
-	// Load the version number.
-	int32 Version = -1;
-	FileReader->Serialize(&Version, sizeof(int32));
-	if (FileReader->IsError())
+	if (GroupMLP_DEPRECATED)
 	{
-		FileReader->Close();
-		UE_LOG(LogNeuralMorphModel, Error, TEXT("Failed to load version"));
-		Empty();
-		return false;
+		TArray<uint8> FileData;
+		UE::NeuralMorphModel::Private::ConvertMLPToFileData(FileData, GroupMLP_DEPRECATED);
+
+		if (!GroupModelData)
+		{
+			GroupModelData = NewObject<UNNEModelData>(this);
+		}
+
+		GroupModelData->Init(TEXT("ubnne"), FileData);
+
+		FileData.Empty();
 	}
 
-	if (Version != 3)
+	// Create models
+
+	TWeakInterfacePtr<INNERuntimeCPU> RuntimeCPU = UE::NNE::GetRuntime<INNERuntimeCPU>(TEXT("NNERuntimeBasicCpu"));
+
+	if (ensureMsgf(RuntimeCPU.IsValid(), TEXT("Could not find requested NNE Runtime")))
 	{
-		UE_LOG(LogNeuralMorphModel, Error, TEXT("The Neural Morph Network file '%s' is of an unknown version (Version=%d)!"), *Filename, Version);
-		FileReader->Close();
-		Empty();
-		return false;
+		if (MainModelData)
+		{
+			MainModel = RuntimeCPU->CreateModelCPU(MainModelData);
+		}
+
+		if (GroupModelData)
+		{
+			GroupModel = RuntimeCPU->CreateModelCPU(GroupModelData);
+		}
 	}
 
-	// Load the header with info.
-	struct FInfoHeader
-	{
-		int32 Mode;						// 0=Local, 1=Global
-		int32 NumInputs;				// How many float inputs in the input layer?
-		int32 NumHiddenLayers;			// How many hidden layers in the network (layers excluding input and output layer)?
-		int32 NumUnitsPerHiddenLayer;	// How many neurons for each hidden layer?
-		int32 NumOutputs;				// The number of units in the output layer.
-		int32 NumMorphsPerBone;			// The number of morph targets per bone, if set Mode == 0. Otherwise ignored.
-		int32 NumBones;					// The number of bones used as input.
-		int32 NumCurves;				// The number of curves used as input.
-		int32 NumFloatsPerCurve;		// The number of floats per curve.
-		int32 NumGroups;				// How many groups are there?
-		int32 NumItemsPerGroup;			// How many items (bones/curves) per group?
-	};
+	// If we are not in the editor then we clear the FileData and FileType since these will be
+	// using additional memory if we are loading from the legacy format.
 
-	FInfoHeader Info;
-	FileReader->Serialize(&Info, sizeof(FInfoHeader));
-	if (FileReader->IsError())
+#if !WITH_EDITOR
+	if (MainModelData)
 	{
-		UE_LOG(LogNeuralMorphModel, Error, TEXT("Failed to load info header!"));
-		FileReader->Close();
-		Empty();
-		return false;
+		MainModelData->ClearFileDataAndFileType();
 	}
 
-	Mode = (Info.Mode == 0) ? ENeuralMorphMode::Local : ENeuralMorphMode::Global;
-	NumMorphsPerBone = Info.NumMorphsPerBone;
-	NumBones = Info.NumBones;
-	NumCurves = Info.NumCurves;
-	NumFloatsPerCurve = Info.NumFloatsPerCurve;
-	NumGroups = Info.NumGroups;
-	NumItemsPerGroup = Info.NumItemsPerGroup;
-
-	// Read the input standard deviation and means.
-	InputMeans.SetNumZeroed(Info.NumInputs);
-	InputStd.SetNumZeroed(Info.NumInputs);
-	FileReader->Serialize(InputStd.GetData(), Info.NumInputs * sizeof(float));
-	if (FileReader->IsError())
+	if (GroupModelData)
 	{
-		UE_LOG(LogNeuralMorphModel, Error, TEXT("Failed to load input standard deviations!"));
-		FileReader->Close();
-		Empty();
-		return false;
+		GroupModelData->ClearFileDataAndFileType();
 	}
+#endif
+}
 
-	FileReader->Serialize(InputMeans.GetData(), Info.NumInputs * sizeof(float));
-	if (FileReader->IsError())
+PRAGMA_ENABLE_DEPRECATION_WARNINGS
+
+void UNeuralMorphNetwork::Empty()
+{
+	InputMeans.Empty();
+	InputStd.Empty();
+
+	if (MainModelData)
 	{
-		UE_LOG(LogNeuralMorphModel, Error, TEXT("Failed to load input means!"));
-		FileReader->Close();
-		Empty();
-		return false;
+		MainModelData->ConditionalBeginDestroy();
+		MainModelData = nullptr;
+		MainModel.Reset();
 	}
 
-	// Load the main network.
-	MainMLP = LoadMLP(*FileReader);
-	if (!MainMLP)
+	if (GroupModelData)
 	{
-		UE_LOG(LogNeuralMorphModel, Error, TEXT("Failed to load main MLP!"));
-		FileReader->Close();
-		Empty();
-		return false;
+		GroupModelData->ConditionalBeginDestroy();
+		GroupModelData = nullptr;
+		GroupModel.Reset();
 	}
 
-	// If we're in local mode and we have groups, load the network for that.
-	if (NumGroups > 0 && Mode == ENeuralMorphMode::Local)
+	Mode = ENeuralMorphMode::Global;
+	NumMorphs = 0;
+	NumMorphsPerBone = 0;
+	NumBones = 0;
+	NumCurves = 0;
+	NumGroups = 0;
+	NumItemsPerGroup = 0;
+	NumFloatsPerCurve = 1;
+}
+
+bool UNeuralMorphNetwork::IsEmpty() const
+{
+	return !MainModel;
+}
+
+bool UNeuralMorphNetwork::Load(const FString& Filename)
+{
+	UE_LOG(LogNeuralMorphModel, Display, TEXT("Loading Neural Morph Network from file '%s'"), *Filename);
+	
+	Empty();
+
+	TArray<uint8> FileData;
+
+	if (FFileHelper::LoadFileToArray(FileData, *Filename))
 	{
-		GroupMLP = LoadMLP(*FileReader);
-		if (!GroupMLP)
+		uint64 Offset = 0;
+
+		if (FileData.Num() < sizeof(uint32) * 2)
 		{
-			UE_LOG(LogNeuralMorphModel, Error, TEXT("Failed to load group MLP!"));
-			FileReader->Close();
-			Empty();
+			UE_LOG(LogNeuralMorphModel, Error, TEXT("The file is not a valid neural morph network file!"));
 			return false;
 		}
-	}
 
-	// Clean up and return the result.
-	const bool bSuccess = FileReader->Close();
-	if (bSuccess)
-	{
-		UE_LOG(LogNeuralMorphModel, Display, TEXT("Successfullly loaded neural morph network from file '%s'"), *FileReader->GetArchiveName());
+		uint32 Magic;
+		UE::NeuralMorphModel::Private::Serialization::Load(Offset, Magic, FileData);
+		if (Magic != UE::NeuralMorphModel::Private::MagicNumber)
+		{
+			UE_LOG(LogNeuralMorphModel, Error, TEXT("The file is not a valid neural morph network file!"));
+			return false;
+		}
+
+		uint32 Version;
+		UE::NeuralMorphModel::Private::Serialization::Load(Offset, Version, FileData);
+		if (Version != UE::NeuralMorphModel::Private::VersionNumber)
+		{
+			UE_LOG(LogNeuralMorphModel, Error, TEXT("The file is not a supported version of neural morph network file!"));
+			return false;
+		}
+
+		// Load General Info
+
+		uint32 ModeValue;
+		UE::NeuralMorphModel::Private::Serialization::Load(Offset, ModeValue, FileData);
+		Mode = (ModeValue == 0) ? ENeuralMorphMode::Local : ENeuralMorphMode::Global;
+		UE::NeuralMorphModel::Private::Serialization::Load(Offset, NumMorphs, FileData);
+		UE::NeuralMorphModel::Private::Serialization::Load(Offset, NumMorphsPerBone, FileData);
+		UE::NeuralMorphModel::Private::Serialization::Load(Offset, NumBones, FileData);
+		UE::NeuralMorphModel::Private::Serialization::Load(Offset, NumCurves, FileData);
+		UE::NeuralMorphModel::Private::Serialization::Load(Offset, NumGroups, FileData);
+		UE::NeuralMorphModel::Private::Serialization::Load(Offset, NumItemsPerGroup, FileData);
+		UE::NeuralMorphModel::Private::Serialization::Load(Offset, NumFloatsPerCurve, FileData);
+
+		// Load Mean / Std
+
+		InputMeans.SetNumUninitialized(GetNumInputs());
+		InputStd.SetNumUninitialized(GetNumInputs());
+
+		UE::NeuralMorphModel::Private::Serialization::Load(Offset, InputMeans, FileData);
+		UE::NeuralMorphModel::Private::Serialization::Load(Offset, InputStd, FileData);
+
+		// Load NNE Runtime Name
+
+		FString RuntimeName;
+		UE::NeuralMorphModel::Private::Serialization::Load(Offset, RuntimeName, FileData);
+
+		// Load Main Network
+
+		if (!MainModelData)
+		{
+			MainModelData = NewObject<UNNEModelData>(this);
+		}
+		
+		uint32 MainModelSize;
+		UE::NeuralMorphModel::Private::Serialization::Load(Offset, MainModelSize, FileData);
+		UE::NeuralMorphModel::Private::Serialization::Align(Offset, 64);
+		MainModelData->Init(TEXT("ubnne"), TConstArrayView<uint8>(FileData).Slice(Offset, MainModelSize));
+
+		Offset += MainModelSize;
+
+		// If we're in local mode and we have groups, load Group Network
+
+		if (NumGroups > 0 && Mode == ENeuralMorphMode::Local)
+		{
+			if (GroupModelData)
+			{
+				GroupModelData = NewObject<UNNEModelData>(this);
+			}
+			
+			uint32 GroupModelSize;
+			UE::NeuralMorphModel::Private::Serialization::Load(Offset, GroupModelSize, FileData);
+			UE::NeuralMorphModel::Private::Serialization::Align(Offset, 64);
+			GroupModelData->Init(TEXT("ubnne"), TConstArrayView<uint8>(FileData).Slice(Offset, GroupModelSize));
+
+			Offset += GroupModelSize;
+		}
+
+		// Create models
+
+		TWeakInterfacePtr<INNERuntimeCPU> RuntimeCPU = UE::NNE::GetRuntime<INNERuntimeCPU>(RuntimeName);
+
+		if (ensureMsgf(RuntimeCPU.IsValid(), TEXT("Could not find requested NNE Runtime")))
+		{
+			if (MainModelData)
+			{
+				MainModel = RuntimeCPU->CreateModelCPU(MainModelData);
+			}
+
+			if (GroupModelData)
+			{
+				GroupModel = RuntimeCPU->CreateModelCPU(GroupModelData);
+			}
+		}
+
+		// Done!
+
+		checkf(Offset == FileData.Num(), TEXT("Did not fully parse file contents"));
+
+		UE_LOG(LogNeuralMorphModel, Display, TEXT("Successfully loaded neural morph network from file '%s'"), *Filename);
+		return true;
 	}
 	else
 	{
-		UE_LOG(LogNeuralMorphModel, Error, TEXT("Failed to close file reader."));
-		Empty();
+		UE_LOG(LogNeuralMorphModel, Error, TEXT("Failed to read file '%s'!"), *Filename);
+		return false;
 	}
-
-	return bSuccess;
 }
 
 UNeuralMorphNetworkInstance* UNeuralMorphNetwork::CreateInstance()
@@ -417,33 +600,77 @@ int32 UNeuralMorphNetwork::GetNumItemsPerGroup() const
 	return NumItemsPerGroup;
 }
 
+int32 UNeuralMorphNetwork::GetNumMainInputs() const
+{
+	return Mode == ENeuralMorphMode::Local ?
+		(NumBones * 6) + (NumCurves * 6) :
+		(NumBones * 6) + NumCurves;
+}
+
+int32 UNeuralMorphNetwork::GetNumMainOutputs() const
+{
+	return Mode == ENeuralMorphMode::Local ?
+		NumMorphsPerBone * (NumBones + NumCurves) :
+		NumMorphs;
+}
+
+int32 UNeuralMorphNetwork::GetNumGroupInputs() const
+{
+	return Mode == ENeuralMorphMode::Local ?
+		NumGroups * 6 :
+		0;
+}
+
+int32 UNeuralMorphNetwork::GetNumGroupOutputs() const
+{
+	return Mode == ENeuralMorphMode::Local ?
+		NumMorphsPerBone * NumGroups :
+		0;
+}
+
 int32 UNeuralMorphNetwork::GetNumInputs() const
 {
-	return MainMLP ? MainMLP->GetNumInputs() : 0;
+	return GetNumMainInputs();
 }
 
 int32 UNeuralMorphNetwork::GetNumOutputs() const
 {
-	int32 NumOutputs = MainMLP ? MainMLP->GetNumOutputs() : 0;
-	NumOutputs += GroupMLP ? GroupMLP->GetNumOutputs() : 0; 
-	return NumOutputs;
+	return GetNumMainOutputs() + GetNumGroupOutputs();
 }
 
+PRAGMA_DISABLE_DEPRECATION_WARNINGS
+
 UNeuralMorphMLP* UNeuralMorphNetwork::GetMainMLP() const
 {
-	return MainMLP.Get();
+	return nullptr;
 }
 
 UNeuralMorphMLP* UNeuralMorphNetwork::GetGroupMLP() const
 {
-	return GroupMLP.Get();
+	return nullptr;
 }
 
+PRAGMA_ENABLE_DEPRECATION_WARNINGS
+
+UE::NNE::IModelCPU* UNeuralMorphNetwork::GetMainModel() const
+{
+	return MainModel.Get();
+}
+
+UE::NNE::IModelCPU* UNeuralMorphNetwork::GetGroupModel() const
+{
+	return GroupModel.Get();
+}
 
 //--------------------------------------------------------------------------
 // UNeuralMorphNetworkInstance
 //--------------------------------------------------------------------------
 
+UNeuralMorphNetworkInstance::UNeuralMorphNetworkInstance() : Super() {}
+UNeuralMorphNetworkInstance::UNeuralMorphNetworkInstance(const FObjectInitializer& ObjectInitializer) : Super(ObjectInitializer) {}
+UNeuralMorphNetworkInstance::UNeuralMorphNetworkInstance(FVTableHelper& Helper) : Super(Helper) {}
+UNeuralMorphNetworkInstance::~UNeuralMorphNetworkInstance() {}
+
 TArrayView<float> UNeuralMorphNetworkInstance::GetInputs()
 { 
 	return TArrayView<float>(Inputs.GetData(), Inputs.Num());
@@ -482,215 +709,64 @@ const UNeuralMorphNetwork& UNeuralMorphNetworkInstance::GetNeuralNetwork() const
 void UNeuralMorphNetworkInstance::Init(UNeuralMorphNetwork* InNeuralNetwork)
 {
 	Network = InNeuralNetwork;
-	const UNeuralMorphMLP* MainMLP = Network->GetMainMLP();
-	const UNeuralMorphMLP* GroupMLP = Network->GetGroupMLP();
 
-	Inputs.SetNumZeroed(Network->GetNumInputs());
-	Outputs.SetNumZeroed(Network->GetNumOutputs());	// This contains concatenated outputs of both main and group MLP's.
+	UE::NNE::IModelCPU* MainNeuralNetwork = Network->GetMainModel();
+	UE::NNE::IModelCPU* GroupNeuralNetwork = Network->GetGroupModel();
 
-	if (GroupMLP)
+	if (ensureMsgf(MainNeuralNetwork, TEXT("Main Neural Network must be valid to create instance")))
 	{
-		GroupInputs.SetNumZeroed(GroupMLP->GetNumInputs());
+		Inputs.SetNumZeroed(Network->GetNumInputs());   // This is just the number of inputs for the main network.
+		Outputs.SetNumZeroed(Network->GetNumOutputs()); // This contains concatenated outputs of both main and group networks.
+		MainInstance = MainNeuralNetwork->CreateModelInstanceCPU();
+		MainInstance->SetInputTensorShapes({ UE::NNE::FTensorShape::Make({ 1, (uint32)Inputs.Num() }) });
 	}
 
-	// Find the largest layer unit size and pre-allocate a buffer of that size.
-	const int32 MainMaxNumInputs = MainMLP ? MainMLP->GetMaxNumLayerInputs() : 0;
-	const int32 MainMaxNumOutputs = MainMLP ? MainMLP->GetMaxNumLayerOutputs() : 0;
-	const int32 GroupMaxNumInputs = GroupMLP ? GroupMLP->GetMaxNumLayerInputs() : 0;
-	const int32 GroupMaxNumOutputs = GroupMLP ? GroupMLP->GetMaxNumLayerOutputs() : 0;
-
-	// Find the maximum of the inputs and outputs of both main and group mlp's.
-	int32 MaxNumUnits = FMath::Max<int32>(MainMaxNumInputs, GroupMaxNumInputs);
-	MaxNumUnits = FMath::Max3<int32>(MaxNumUnits, MainMaxNumOutputs, GroupMaxNumOutputs);
-
-	TempInputArray.SetNumZeroed(MaxNumUnits);
-	TempOutputArray.SetNumZeroed(MaxNumUnits);
+	if (GroupNeuralNetwork)
+	{
+		GroupInputs.SetNumZeroed(Network->GetNumGroupInputs());
+		GroupInstance = GroupNeuralNetwork->CreateModelInstanceCPU();
+		GroupInstance->SetInputTensorShapes({ UE::NNE::FTensorShape::Make({ 1, (uint32)GroupInputs.Num() }) });
+	}
 }
 
-void UNeuralMorphNetworkInstance::RunGlobalModel(const FRunSettings& RunSettings)
+void UNeuralMorphNetworkInstance::Run()
 {
-	float* RESTRICT TempInputBuffer = RunSettings.TempInputBuffer;
-	float* RESTRICT TempOutputBuffer = RunSettings.TempOutputBuffer;
+	TRACE_CPUPROFILER_EVENT_SCOPE(UNeuralMorphNetwork::Run);
 
-	checkfSlow(Network->GetNumFloatsPerCurve() == 1, TEXT("Expecting the number of floats per curve to be 1 in global mode."));
+	UE::NNE::IModelCPU* MainNeuralNetwork = Network->GetMainModel();
+	UE::NNE::IModelCPU* GroupNeuralNetwork = Network->GetGroupModel();
 
-	const UNeuralMorphMLP& MainMLP = *Network->GetMainMLP();
+	TArrayView<float> MainOutputs = TArrayView<float>(Outputs.GetData(), Network->GetNumMainOutputs());
 
-	const int32 NumLayers = MainMLP.GetNumLayers();
-	for (int32 LayerIndex = 0; LayerIndex < NumLayers; ++LayerIndex)
+	if (Network->GetMode() == ENeuralMorphMode::Global)
 	{
-		const UNeuralMorphMLPLayer& CurLayer = MainMLP.GetLayer(LayerIndex);
-		const int32 NumLayerInputs = CurLayer.NumInputs;
-		const int32 NumLayerOutputs = CurLayer.NumOutputs;
-
-		// Copy the inputs to the temp input buffer for the first layer.
-		if (LayerIndex == 0)
-		{
-			const float* const RESTRICT NetworkInputs = RunSettings.InputBuffer;
-			for (int32 Index = 0; Index < NumLayerInputs; ++Index)
-			{
-				TempInputBuffer[Index] = NetworkInputs[Index];
-			}
-		}
+		checkfSlow(Network->GetNumFloatsPerCurve() == 1, TEXT("Expecting the number of floats per curve to be 1 in global mode."));
 
-		// If we reached the last layer, output to the final buffer.
-		if (LayerIndex == NumLayers - 1)
-		{
-			TempOutputBuffer = RunSettings.OutputBuffer;
-		}
-
-		#if NEURALMORPHMODEL_USE_ISPC
-			const float* const RESTRICT LayerWeights = CurLayer.Weights.GetData();
-			const float* const RESTRICT LayerBiases = CurLayer.Biases.GetData();
-			ispc::MorphNeuralNetwork_LayerForward(TempOutputBuffer, TempInputBuffer, LayerWeights, LayerBiases, NumLayerInputs, NumLayerOutputs);
-		#else
-			// Init the outputs to the bias.
-			const float* const RESTRICT LayerBiases = CurLayer.Biases.GetData();
-			for (int32 Index = 0; Index < NumLayerOutputs; ++Index)
-			{
-				TempOutputBuffer[Index] = LayerBiases[Index];
-			}
-
-			// Multiply layer inputs with the weights.
-			const float* const RESTRICT LayerWeights = CurLayer.Weights.GetData();
-			for (int32 InputIndex = 0; InputIndex < NumLayerInputs; ++InputIndex)
-			{
-				const float InputValue = TempInputBuffer[InputIndex];
-				const int32 Offset = InputIndex * NumLayerOutputs;
-				for (int32 OutputIndex = 0; OutputIndex < NumLayerOutputs; ++OutputIndex)
-				{
-					TempOutputBuffer[OutputIndex] += InputValue * LayerWeights[Offset + OutputIndex];
-				}
-			}
-
-			// Apply ELU activation.
-			for (int32 Index = 0; Index < NumLayerOutputs; ++Index)
-			{
-				const float X = TempOutputBuffer[Index];
-				TempOutputBuffer[Index] = (X > 0.0f) ? X : FMath::InvExpApprox(-X) - 1.0f;
-			}
-		#endif
-
-		// The outputs are now input to the next layer.
-		if (LayerIndex < NumLayers - 1)
-		{
-			// Swap the inputs and outputs, as the outputs are the input to the next layer.
-			float* const SwapTemp = TempInputBuffer;
-			TempInputBuffer = TempOutputBuffer;
-			TempOutputBuffer = SwapTemp;
-		}
+		MainInstance->RunSync(
+			{ { (void*)Inputs.GetData(), Inputs.Num() * sizeof(float) } },
+			{ { (void*)MainOutputs.GetData(), MainOutputs.Num() * sizeof(float) } });
 	}
-}
-
-void UNeuralMorphNetworkInstance::RunLocalMLP(UNeuralMorphMLP& MLP, const FRunSettings& RunSettings)
-{
-	float* RESTRICT TempInputBuffer = RunSettings.TempInputBuffer;
-	float* RESTRICT TempOutputBuffer = RunSettings.TempOutputBuffer;
-
-	checkfSlow(Network->GetNumFloatsPerCurve() == 6, TEXT("Expecting num floats per curve to be 6 in local mode."));
-
-	const int32 NumLayers = MLP.GetNumLayers();
-	for (int32 LayerIndex = 0; LayerIndex < NumLayers; ++LayerIndex)
+	else if (Network->GetMode() == ENeuralMorphMode::Local)
 	{
-		const UNeuralMorphMLPLayer& CurLayer = MLP.GetLayer(LayerIndex);
-		const int32 NumInputsPerBlock = CurLayer.NumInputs;
-		const int32 NumBlocks = CurLayer.Depth;
+		checkfSlow(Network->GetNumFloatsPerCurve() == 6, TEXT("Expecting num floats per curve to be 6 in local mode."));
 
-		// Copy the inputs to the temp buffer.
-		if (LayerIndex == 0)
-		{
-			const float* const RESTRICT NetworkInputs = RunSettings.InputBuffer;
-			const int32 NumInputs = CurLayer.NumInputs * CurLayer.Depth;
-			for (int32 Index = 0; Index < NumInputs; ++Index)
-			{
-				TempInputBuffer[Index] = NetworkInputs[Index];
-			}
-		}
+		MainInstance->RunSync(
+			{ { (void*)Inputs.GetData(), Inputs.Num() * sizeof(float) } },
+			{ { (void*)MainOutputs.GetData(), MainOutputs.Num() * sizeof(float) } });
 
-		// If we reached the last layer, output to the final buffer.
-		if (LayerIndex == NumLayers - 1)
+		if (GroupNeuralNetwork)
 		{
-			TempOutputBuffer = RunSettings.OutputBuffer;
-		}
+			TArrayView<float> GroupOutputs = TArrayView<float>(
+				Outputs.GetData() + Network->GetNumMainOutputs(),
+				Network->GetNumGroupOutputs());
 
-		// Init the output buffer to the bias values.
-		const float* const RESTRICT LayerBiases = CurLayer.Biases.GetData();
-		const int32 NumOutputsPerBlock = CurLayer.NumOutputs;
-		const int32 NumOutputs = CurLayer.NumOutputs * CurLayer.Depth;
-		for (int32 Index = 0; Index < NumOutputs; ++Index)
-		{
-			TempOutputBuffer[Index] = LayerBiases[Index];
+			GroupInstance->RunSync(
+				{ { (void*)GroupInputs.GetData(), GroupInputs.Num() * sizeof(float) } },
+				{ { (void*)GroupOutputs.GetData(), GroupOutputs.Num() * sizeof(float) } });
 		}
-
-		const float* const RESTRICT LayerWeights = CurLayer.Weights.GetData();
-		for (int32 BlockIndex = 0; BlockIndex < NumBlocks; ++BlockIndex)
-		{
-			const int32 BlockOutputOffset = BlockIndex * NumOutputsPerBlock;
-			const int32 BlockInputOffset = BlockIndex * NumInputsPerBlock;
-
-			// Multiply layer inputs with the weights.
-			const int32 WeightOffset = BlockIndex * (NumInputsPerBlock * NumOutputsPerBlock);
-			for (int32 InputIndex = 0; InputIndex < NumInputsPerBlock; ++InputIndex)
-			{
-				const float InputValue = TempInputBuffer[BlockInputOffset + InputIndex];
-				const int32 InputOffset = InputIndex * NumOutputsPerBlock;
-				for (int32 OutputIndex = 0; OutputIndex < NumOutputsPerBlock; ++OutputIndex)
-				{
-					TempOutputBuffer[BlockOutputOffset + OutputIndex] += InputValue * LayerWeights[WeightOffset + InputOffset + OutputIndex];
-				}
-			}
-		} // For all blocks.
-
-		// Apply ELU activation.
-		#if NEURALMORPHMODEL_USE_ISPC
-			ispc::MorphNeuralNetwork_Activation_ELU(TempOutputBuffer, NumOutputs);
-		#else
-			for (int32 OutputIndex = 0; OutputIndex < NumOutputs; ++OutputIndex)
-			{
-				const float X = TempOutputBuffer[OutputIndex];
-				TempOutputBuffer[OutputIndex] = (X > 0.0f) ? X : FMath::InvExpApprox(-X) - 1.0f;
-			}
-		#endif
-
-		// The outputs are now input to the next layer.
-		if (LayerIndex < NumLayers - 1)
-		{
-			// Swap the inputs and outputs, as the outputs are the input to the next layer.
-			float* const SwapTemp = TempInputBuffer;
-			TempInputBuffer = TempOutputBuffer;
-			TempOutputBuffer = SwapTemp;
-		}
-	} // For all layers.
-}
-
-void UNeuralMorphNetworkInstance::Run()
-{	
-	TRACE_CPUPROFILER_EVENT_SCOPE(UNeuralMorphNetwork::Run)
-
-	// Setup the buffer pointers.
-	FRunSettings RunSettings;
-	RunSettings.TempInputBuffer  = TempInputArray.GetData();
-	RunSettings.TempOutputBuffer = TempOutputArray.GetData();
-	RunSettings.InputBuffer		 = Inputs.GetData();
-	RunSettings.OutputBuffer	 = Outputs.GetData();
-
-	// Run the network.
-	if (Network->GetMode() == ENeuralMorphMode::Global)
-	{
-		RunGlobalModel(RunSettings);
 	}
 	else
 	{
-		checkSlow(Network->GetMode() == ENeuralMorphMode::Local)
-		RunLocalMLP(*Network->GetMainMLP(), RunSettings);
-		if (Network->GetGroupMLP())
-		{
-			FRunSettings GroupRunSettings;
-			GroupRunSettings.TempInputBuffer	= TempInputArray.GetData();
-			GroupRunSettings.TempOutputBuffer	= TempOutputArray.GetData();
-			GroupRunSettings.InputBuffer		= GroupInputs.GetData();
-			GroupRunSettings.OutputBuffer		= Outputs.GetData() + Network->GetMainMLP()->GetNumOutputs();	// Output values after the main network's outputs.
-			RunLocalMLP(*Network->GetGroupMLP(), GroupRunSettings);
-		}
+		checkf(false, TEXT("Unknown Mode"));
 	}
 }
diff --git a/Engine/Plugins/Animation/MLDeformer/NeuralMorphModel/Source/NeuralMorphModel/Private/NeuralMorphNetwork.ispc b/Engine/Plugins/Animation/MLDeformer/NeuralMorphModel/Source/NeuralMorphModel/Private/NeuralMorphNetwork.ispc
deleted file mode 100644
diff --git a/Engine/Plugins/Animation/MLDeformer/NeuralMorphModel/Source/NeuralMorphModel/Public/NeuralMorphModel.h b/Engine/Plugins/Animation/MLDeformer/NeuralMorphModel/Source/NeuralMorphModel/Public/NeuralMorphModel.h
--- a/Engine/Plugins/Animation/MLDeformer/NeuralMorphModel/Source/NeuralMorphModel/Public/NeuralMorphModel.h
+++ b/Engine/Plugins/Animation/MLDeformer/NeuralMorphModel/Source/NeuralMorphModel/Public/NeuralMorphModel.h
@@ -17,15 +17,6 @@ class UMLDeformerModelInstance;
 class USkeleton;
 struct FExternalMorphSet;
 
-/** 
- * Specify whether we want to use ISPC if available.
- * You can comment out this line, to force disable the use of ISPC in this plugin.
- */
-#define NEURALMORPHMODEL_USE_ISPC INTEL_ISPC
-#if !defined(NEURALMORPHMODEL_USE_ISPC)
-	#define NEURALMORPHMODEL_USE_ISPC 0
-#endif
-
 // Declare our log category.
 NEURALMORPHMODEL_API DECLARE_LOG_CATEGORY_EXTERN(LogNeuralMorphModel, Log, All);
 
diff --git a/Engine/Plugins/Animation/MLDeformer/NeuralMorphModel/Source/NeuralMorphModel/Public/NeuralMorphNetwork.h b/Engine/Plugins/Animation/MLDeformer/NeuralMorphModel/Source/NeuralMorphModel/Public/NeuralMorphNetwork.h
--- a/Engine/Plugins/Animation/MLDeformer/NeuralMorphModel/Source/NeuralMorphModel/Public/NeuralMorphNetwork.h
+++ b/Engine/Plugins/Animation/MLDeformer/NeuralMorphModel/Source/NeuralMorphModel/Public/NeuralMorphNetwork.h
@@ -6,120 +6,72 @@
 #include "NeuralMorphModel.h"
 #include "NeuralMorphNetwork.generated.h"
 
-class UNeuralMorphNetworkInstance;
+class UNNEModelData;
+
+namespace UE::NNE
+{
+	class IModelCPU;
+	class IModelInstanceCPU;
+}
 
+class UNeuralMorphNetworkInstance;
 
-/** A fully connected layer, which contains the weights and biases for those connections. */ 
 UCLASS()
-class NEURALMORPHMODEL_API UNeuralMorphMLPLayer
+class UE_DEPRECATED(5.4, "Neural Morph MLP no longer used for inference.") NEURALMORPHMODEL_API UNeuralMorphMLPLayer
 	: public UObject
 {
 	GENERATED_BODY()
 
 public:
-	/** The weight matrix number of inputs (rows). */
+	
 	UPROPERTY()
 	int32 NumInputs = 0;
 
-	/** The weight matrix number of outputs (columns). */
 	UPROPERTY()
 	int32 NumOutputs = 0;
 
-	/** The number of instances of inputs and outputs. */
 	UPROPERTY()
 	int32 Depth = 1;
 
-	/** A 2D array of weights. The number of weights equals NumRows x NumColumns x Depth. */
 	UPROPERTY()
 	TArray<float> Weights;
 
-	/** The biases. The number of biases will be the same as the number of columns multiplied by the depth. */
 	UPROPERTY()
 	TArray<float> Biases;
 };
 
 
-/**
- * An MLP neural network.
- */
 UCLASS()
-class NEURALMORPHMODEL_API UNeuralMorphMLP
+class UE_DEPRECATED(5.4, "Neural Morph MLP longer used for inference.") NEURALMORPHMODEL_API UNeuralMorphMLP
 	: public UObject
 {
 	GENERATED_BODY()
 
 public:
-	/** Clear the network, getting rid of all weights and biases. */
-	void Empty();
+	
+	PRAGMA_DISABLE_DEPRECATION_WARNINGS
 
-	/** 
-	 * Check if the network is empty or not.
-	 * If it is empty, it means it hasn't been loaded, and cannot do anything.
-	 * Empty means it has no layers.
-	 * @return Returns true if empty, otherwise false is returned.
-	 */
+	void Empty();
 	bool IsEmpty() const;
-
-	/**
-	 * Load the network from a file on disk.
-	 * When loading fails, the network will be emptied.
-	 * @param Archive The archive to load from.
-	 * @return Returns true when successfully loaded this model, otherwise false is returned.
-	 */
 	bool Load(FArchive& FileReader);
-
-	/**
-	 * Get the number of inputs, which is the number of floats the network takes as input.
-	 * @result The number of input floats to the network.
-	 */
 	int32 GetNumInputs() const;
-
-	/**
-	 * Get the number of outputs, which is the number of floats the network will output.
-	 * @return The number of floats that the network outputs.
-	 */
 	int32 GetNumOutputs() const;
-
-	/**
-	 * Get the number of network layers.
-	 * This equals to the number of hidden layers plus one.
-	 * @return The number of network layers.
-	 */
 	int32 GetNumLayers() const;
-
-	/**
-	 * Get a given network layer.
-	 * @return A reference to the layer, which will contain the weights and biases.
-	 */
 	UNeuralMorphMLPLayer& GetLayer(int32 Index) const;
-
-	/**
-	 * Get the maximum number of inputs for all layers in the network.
-	 * This is used to pre-allocate a temp buffer used to pass the data through the layers.
-	 * @return The maximum number of inputs of all layers.
-	 */
 	int32 GetMaxNumLayerInputs() const;
-
-	/**
-	 * Get the maximum number of outputs for all layers in the network.
-	 * This is used to pre-allocate a temp buffer used to pass the data through the layers.
-	 * @return The maximum number of outputs of all layers.
-	 */
 	int32 GetMaxNumLayerOutputs() const;
 
 private:
-	/**
-	 * Load an MLP fully connected layer from a given archive.
-	 * @param Archive The archive to read from.
-	 */
+	
 	UNeuralMorphMLPLayer* LoadNetworkLayer(FArchive& Archive);
 
 private:
-	/** The network weights and biases of the main network. */
+
 	UPROPERTY()
 	TArray<TObjectPtr<UNeuralMorphMLPLayer>> Layers;
-};
 
+	PRAGMA_ENABLE_DEPRECATION_WARNINGS
+};
 
 /**
  * The specialized neural network for the Neural Morph Model.
@@ -135,6 +87,16 @@ class NEURALMORPHMODEL_API UNeuralMorphNetwork
 	GENERATED_BODY()
 
 public:
+
+	// These are required because of the use of TSharedPtr
+	UNeuralMorphNetwork();
+	UNeuralMorphNetwork(const FObjectInitializer& ObjectInitializer);
+	UNeuralMorphNetwork(FVTableHelper& Helper);
+	virtual ~UNeuralMorphNetwork();
+
+	/** Used to convert the legacy neural network format to the NNE format */
+	virtual void PostLoad() override;
+
 	/** Clear the network, getting rid of all weights and biases. */
 	void Empty();
 
@@ -160,9 +122,33 @@ public:
 	 */
 	UNeuralMorphNetworkInstance* CreateInstance();
 
+	/**
+	 * Get the number of network inputs for the main network.
+	 * @return The number of inputs.
+	 */
+	int32 GetNumMainInputs() const;
+
+	/**
+	 * Get the number of network outputs for the main network.
+	 * @return The number of outputs.
+	 */
+	int32 GetNumMainOutputs() const;
+
+	/**
+	 * Get the number of network inputs for the group network.
+	 * @return The number of inputs.
+	 */
+	int32 GetNumGroupInputs() const;
+
+	/**
+	 * Get the number of network outputs for the group network.
+	 * @return The number of outputs.
+	 */
+	int32 GetNumGroupOutputs() const;
+
 	/**
 	 * Get the number of network inputs.
-	 * This is the sum of the number of main network inputs and group network inputs.
+	 * This is the number of main network inputs.
 	 * @return The number of inputs.
 	 */
 	int32 GetNumInputs() const;
@@ -233,38 +219,59 @@ public:
 	 */
 	int32 GetNumFloatsPerCurve() const;
 
+	PRAGMA_DISABLE_DEPRECATION_WARNINGS
+
+	UE_DEPRECATED(5.4, "MLP Object no longer available. Please use GetMainNeuralNetwork instead.")
+	UNeuralMorphMLP* GetMainMLP() const;
+
+	UE_DEPRECATED(5.4, "MLP Object no longer available. Please use GetGroupNeuralNetwork instead.")
+	UNeuralMorphMLP* GetGroupMLP() const;
+
+	PRAGMA_ENABLE_DEPRECATION_WARNINGS
+
 	/**
-	 * Get the main MLP.
-	 * This is the MLP (or actually a Multi-MLP) that contains a small network for every bone or curve.
+	 * Get the main Model.
+	 * This is the Model (or actually a Multi-Model) that contains a small network for every bone or curve.
 	 * Each small network will output a set of morph target weights.
-	 * @return A pointer to the main MLP.
+	 * @return A pointer to the main Model.
 	 */
-	UNeuralMorphMLP* GetMainMLP() const;
+	UE::NNE::IModelCPU* GetMainModel() const;
 
 	/**
-	 * Get the MLP used for group of bones and curves.
-	 * This MLP (or actually a Multi-MLP) contains small networks for the bone or curve groups.
+	 * Get the Model used for group of bones and curves.
+	 * This Model (or actually a Multi-Model) contains small networks for the bone or curve groups.
 	 * Each group will generate a set of morph target weights.
-	 * @return A pointer to the MLP used for the bone and curve groups.
+	 * @return A pointer to the Model used for the bone and curve groups.
 	 */
-	UNeuralMorphMLP* GetGroupMLP() const;
+	UE::NNE::IModelCPU* GetGroupModel() const;
 
 private:
-	/**
-	 * Load an MLP network from the currently already opened archive.
-	 * @param Archive The archive to load from.
-	 * @return Returns a pointer to the loaded MLP, or a nullptr in case of a failure.
-	 */
-	UNeuralMorphMLP* LoadMLP(FArchive& FileReader);
 
-private:
-	/** The MLP that acts as main network. */
+	PRAGMA_DISABLE_DEPRECATION_WARNINGS
+
+	UE_DEPRECATED(5.4, "MLP Object no longer available. Please use MainModel instead.")
+	UPROPERTY(meta = (DeprecatedProperty, DeprecationMessage = "Network format changed. Please use MainModel instead."))
+	TObjectPtr<UNeuralMorphMLP> MainMLP_DEPRECATED;
+
+	UE_DEPRECATED(5.4, "MLP Object no longer available. Please use GroupModel instead.")
+	UPROPERTY(meta = (DeprecatedProperty, DeprecationMessage = "Network format changed. Please use GroupModel instead."))
+	TObjectPtr<UNeuralMorphMLP> GroupMLP_DEPRECATED;
+
+	PRAGMA_ENABLE_DEPRECATION_WARNINGS
+
+	/** The Model Data for the Main Network */
 	UPROPERTY()
-	TObjectPtr<UNeuralMorphMLP> MainMLP;
+	TObjectPtr<UNNEModelData> MainModelData;
 
-	/** The MLP for the bone and curve groups, when in local mode. */
+	/** The Neural Network that acts as main network. */
+	TSharedPtr<UE::NNE::IModelCPU> MainModel;
+
+	/** The Model Data for the Group Network */
 	UPROPERTY()
-	TObjectPtr<UNeuralMorphMLP> GroupMLP;
+	TObjectPtr<UNNEModelData> GroupModelData;
+
+	/** The Neural Network for the bone and curve groups, when in local mode. */
+	TSharedPtr<UE::NNE::IModelCPU> GroupModel;
 
 	/** The means of the input values, used to normalize inputs. */
 	UPROPERTY()
@@ -278,6 +285,10 @@ private:
 	UPROPERTY()
 	ENeuralMorphMode Mode = ENeuralMorphMode::Global;
 
+	/** The total number of morph targets, if set Mode == Global, otherwise ignored. */
+	UPROPERTY()
+	int32 NumMorphs = 0;
+
 	/** The number of morph targets per bone, if set Mode == Local, otherwise ignored. */
 	UPROPERTY()
 	int32 NumMorphsPerBone = 0;
@@ -317,6 +328,13 @@ class NEURALMORPHMODEL_API UNeuralMorphNetworkInstance
 	GENERATED_BODY()
 
 public:
+
+	// These are required because of the use of TSharedPtr
+	UNeuralMorphNetworkInstance();
+	UNeuralMorphNetworkInstance(const FObjectInitializer& ObjectInitializer);
+	UNeuralMorphNetworkInstance(FVTableHelper& Helper);
+	virtual ~UNeuralMorphNetworkInstance();
+
 	/**
 	 * Get the network input buffer.
 	 * @return The array view of floats that represents the main network inputs.
@@ -352,26 +370,6 @@ public:
 	void Run();
 
 private:
-	struct FRunSettings
-	{
-		float* TempInputBuffer = nullptr;
-		float* TempOutputBuffer = nullptr;
-		const float* InputBuffer = nullptr;
-		float* OutputBuffer = nullptr;
-	};
-
-	/**
- 	 * Run inference using the global model mode, which is one fully connected MLP.
-	 * @param RunSettings The settings used to run the model. This specifies a set of buffers.
-	 */
-	void RunGlobalModel(const FRunSettings& RunSettings);
-
-	/**
- 	 * Run a local mlp, which internally is acting like a multi-mlp (one mlp per bone or curve).
-	 * @param MLP The MLP to run.
-	 * @param RunSettings The settings used to run the model. This specifies a set of buffers.
-	 */
-	void RunLocalMLP(UNeuralMorphMLP& MLP, const FRunSettings& RunSettings);
 
 	/**
  	 * Initialize this instance using a given network object.
@@ -381,20 +379,20 @@ private:
 	void Init(UNeuralMorphNetwork* InNeuralNetwork);
 
 private:
-	/** The input values for the main MLP. */
+	/** The input values for the main network. */
 	TArray<float> Inputs;
 
-	/** The input values for the group MLP. */
+	/** The input values for the group network. */
 	TArray<float> GroupInputs;
-
+	
 	/** The output values. */
 	TArray<float> Outputs;
 
-	/** A pre-allocated temp buffer for inputs. */
-	TArray<float> TempInputArray;
+	/** The inference instance data for the main Neural Network */
+	TSharedPtr<UE::NNE::IModelInstanceCPU> MainInstance;
 
-	/** A pre-allocated temp buffer for outputs. */
-	TArray<float> TempOutputArray;
+	/** The inference instance data for the group Neural Network */
+	TSharedPtr<UE::NNE::IModelInstanceCPU> GroupInstance;
 
 	/** The neural network this is an instance of. */
 	UPROPERTY(Transient)
diff --git a/Engine/Plugins/Animation/MLDeformer/NeuralMorphModel/Source/NeuralMorphModelEditor/Private/NeuralMorphEditorModel.cpp b/Engine/Plugins/Animation/MLDeformer/NeuralMorphModel/Source/NeuralMorphModelEditor/Private/NeuralMorphEditorModel.cpp
--- a/Engine/Plugins/Animation/MLDeformer/NeuralMorphModel/Source/NeuralMorphModelEditor/Private/NeuralMorphEditorModel.cpp
+++ b/Engine/Plugins/Animation/MLDeformer/NeuralMorphModel/Source/NeuralMorphModelEditor/Private/NeuralMorphEditorModel.cpp
@@ -93,7 +93,7 @@ namespace UE::NeuralMorphModel
 	bool FNeuralMorphEditorModel::IsTrained() const
 	{
 		UNeuralMorphNetwork* MorphNetwork = GetNeuralMorphModel()->GetNeuralMorphNetwork();
-		if (MorphNetwork != nullptr && MorphNetwork->GetMainMLP() != nullptr)
+		if (MorphNetwork != nullptr && MorphNetwork->GetMainModel() != nullptr)
 		{
 			return true;
 		}
