[MLDeformer] Some fixes in the new NeuralMorphNetwork for the local mode of the neural morph model.

#rnx

diff --git a/Engine/Plugins/Animation/MLDeformer/NeuralMorphModel/Source/NeuralMorphModel/Private/NeuralMorphNetwork.cpp b/Engine/Plugins/Animation/MLDeformer/NeuralMorphModel/Source/NeuralMorphModel/Private/NeuralMorphNetwork.cpp
--- a/Engine/Plugins/Animation/MLDeformer/NeuralMorphModel/Source/NeuralMorphModel/Private/NeuralMorphNetwork.cpp
+++ b/Engine/Plugins/Animation/MLDeformer/NeuralMorphModel/Source/NeuralMorphModel/Private/NeuralMorphNetwork.cpp
@@ -504,7 +504,7 @@ bool UNeuralMorphNetwork::Load(const FString& Filename)
 
 		if (NumGroups > 0 && Mode == ENeuralMorphMode::Local)
 		{
-			if (GroupModelData)
+			if (!GroupModelData)
 			{
 				GroupModelData = NewObject<UNNEModelData>(this);
 			}
@@ -617,7 +617,7 @@ int32 UNeuralMorphNetwork::GetNumMainOutputs() const
 int32 UNeuralMorphNetwork::GetNumGroupInputs() const
 {
 	return Mode == ENeuralMorphMode::Local ?
-		NumGroups * 6 :
+		NumGroups * 6 * NumItemsPerGroup :
 		0;
 }
 
