GetAssetRegistryTags deprecation: Make the old GetAssetRegistryTags and related functions deprecated.
Upgrade the overrides and callsites in all of Epic's code.
Every class that can be subclassed by licensees that has its override change needs to keep the old version of the function as deprecated. Otherwise subclasses that call Super::GetAssetRegistryTags will have a compile error instead of a deprecation warning.
Several classes had their own extension method for GetAssetRegistryTags being passed on to user data classes. Updated all of those to take FAssetRegistryTagsContext just like GetAssetRegistryTags does.
#rnx

diff --git a/Engine/Plugins/2D/Paper2D/Source/Paper2D/Classes/PaperSprite.h b/Engine/Plugins/2D/Paper2D/Source/Paper2D/Classes/PaperSprite.h
--- a/Engine/Plugins/2D/Paper2D/Source/Paper2D/Classes/PaperSprite.h
+++ b/Engine/Plugins/2D/Paper2D/Source/Paper2D/Classes/PaperSprite.h
@@ -181,6 +181,8 @@ public:
 	virtual void Serialize(FArchive& Ar) override;
 	virtual void PostLoad() override;
 #if WITH_EDITOR
+	virtual void GetAssetRegistryTags(FAssetRegistryTagsContext Context) const override;
+	UE_DEPRECATED(5.4, "Implement the version that takes FAssetRegistryTagsContext instead.")
 	virtual void GetAssetRegistryTags(TArray<FAssetRegistryTag>& OutTags) const override;
 	virtual void PostEditChangeProperty(FPropertyChangedEvent& PropertyChangedEvent) override;
 #endif
diff --git a/Engine/Plugins/2D/Paper2D/Source/Paper2D/Classes/PaperSpriteAtlas.h b/Engine/Plugins/2D/Paper2D/Source/Paper2D/Classes/PaperSpriteAtlas.h
--- a/Engine/Plugins/2D/Paper2D/Source/Paper2D/Classes/PaperSpriteAtlas.h
+++ b/Engine/Plugins/2D/Paper2D/Source/Paper2D/Classes/PaperSpriteAtlas.h
@@ -124,6 +124,8 @@ public:
 	virtual void PostInitProperties() override;
 	virtual void PostDuplicate(bool bDuplicateForPIE) override;
 #endif
+	virtual void GetAssetRegistryTags(FAssetRegistryTagsContext Context) const override;
+	UE_DEPRECATED(5.4, "Implement the version that takes FAssetRegistryTagsContext instead.")
 	virtual void GetAssetRegistryTags(TArray<FAssetRegistryTag>& OutTags) const override;
 	// End of UObject interface
 };
diff --git a/Engine/Plugins/2D/Paper2D/Source/Paper2D/Private/PaperSprite.cpp b/Engine/Plugins/2D/Paper2D/Source/Paper2D/Private/PaperSprite.cpp
--- a/Engine/Plugins/2D/Paper2D/Source/Paper2D/Private/PaperSprite.cpp
+++ b/Engine/Plugins/2D/Paper2D/Source/Paper2D/Private/PaperSprite.cpp
@@ -2,6 +2,7 @@
 
 #include "PaperSprite.h"
 #include "BodySetupEnums.h"
+#include "UObject/AssetRegistryTagsContext.h"
 #include "UObject/ConstructorHelpers.h"
 #include "Materials/MaterialInterface.h"
 #include "Engine/CollisionProfile.h"
@@ -1647,11 +1648,18 @@ FVector2D UPaperSprite::GetPivotPosition() const
 
 void UPaperSprite::GetAssetRegistryTags(TArray<FAssetRegistryTag>& OutTags) const
 {
+	PRAGMA_DISABLE_DEPRECATION_WARNINGS;
 	Super::GetAssetRegistryTags(OutTags);
+	PRAGMA_ENABLE_DEPRECATION_WARNINGS;
+}
+
+void UPaperSprite::GetAssetRegistryTags(FAssetRegistryTagsContext Context) const
+{
+	Super::GetAssetRegistryTags(Context);
 
 	if (AtlasGroup != nullptr)
 	{
-		OutTags.Add(FAssetRegistryTag(TEXT("AtlasGroupGUID"), AtlasGroup->AtlasGUID.ToString(EGuidFormats::Digits), FAssetRegistryTag::TT_Hidden));
+		Context.AddTag(FAssetRegistryTag(TEXT("AtlasGroupGUID"), AtlasGroup->AtlasGUID.ToString(EGuidFormats::Digits), FAssetRegistryTag::TT_Hidden));
 	}
 }
 
diff --git a/Engine/Plugins/2D/Paper2D/Source/Paper2D/Private/PaperSpriteAtlas.cpp b/Engine/Plugins/2D/Paper2D/Source/Paper2D/Private/PaperSpriteAtlas.cpp
--- a/Engine/Plugins/2D/Paper2D/Source/Paper2D/Private/PaperSpriteAtlas.cpp
+++ b/Engine/Plugins/2D/Paper2D/Source/Paper2D/Private/PaperSpriteAtlas.cpp
@@ -2,6 +2,7 @@
 
 #include "PaperSpriteAtlas.h"
 #include "Engine/Texture.h"
+#include "UObject/AssetRegistryTagsContext.h"
 
 #include UE_INLINE_GENERATED_CPP_BY_NAME(PaperSpriteAtlas)
 
@@ -25,10 +26,17 @@ UPaperSpriteAtlas::UPaperSpriteAtlas(const FObjectInitializer& ObjectInitializer
 
 void UPaperSpriteAtlas::GetAssetRegistryTags(TArray<FAssetRegistryTag>& OutTags) const
 {
+	PRAGMA_DISABLE_DEPRECATION_WARNINGS;
 	Super::GetAssetRegistryTags(OutTags);
+	PRAGMA_ENABLE_DEPRECATION_WARNINGS;
+}
+
+void UPaperSpriteAtlas::GetAssetRegistryTags(FAssetRegistryTagsContext Context) const
+{
+	Super::GetAssetRegistryTags(Context);
 
 #if WITH_EDITORONLY_DATA
-	OutTags.Add(FAssetRegistryTag(GET_MEMBER_NAME_CHECKED(UPaperSpriteAtlas, AtlasDescription), AtlasDescription, FAssetRegistryTag::TT_Hidden));
+	Context.AddTag(FAssetRegistryTag(GET_MEMBER_NAME_CHECKED(UPaperSpriteAtlas, AtlasDescription), AtlasDescription, FAssetRegistryTag::TT_Hidden));
 #endif
 }
 
diff --git a/Engine/Plugins/2D/Paper2D/Source/Paper2D/Private/PaperTileMap.cpp b/Engine/Plugins/2D/Paper2D/Source/Paper2D/Private/PaperTileMap.cpp
--- a/Engine/Plugins/2D/Paper2D/Source/Paper2D/Private/PaperTileMap.cpp
+++ b/Engine/Plugins/2D/Paper2D/Source/Paper2D/Private/PaperTileMap.cpp
@@ -2,6 +2,7 @@
 
 #include "PaperTileMap.h"
 #include "BodySetupEnums.h"
+#include "UObject/AssetRegistryTagsContext.h"
 #include "UObject/ConstructorHelpers.h"
 #include "Materials/MaterialInterface.h"
 #include "EditorFramework/AssetImportData.h"
@@ -258,13 +259,20 @@ void UPaperTileMap::ValidateSelectedLayerIndex()
 
 #if WITH_EDITORONLY_DATA
 void UPaperTileMap::GetAssetRegistryTags(TArray<FAssetRegistryTag>& OutTags) const
+{
+	PRAGMA_DISABLE_DEPRECATION_WARNINGS;
+	Super::GetAssetRegistryTags(OutTags);
+	PRAGMA_ENABLE_DEPRECATION_WARNINGS;
+}
+
+void UPaperTileMap::GetAssetRegistryTags(FAssetRegistryTagsContext Context) const
 {
 	if (AssetImportData)
 	{
-		OutTags.Add( FAssetRegistryTag(SourceFileTagName(), AssetImportData->GetSourceData().ToJson(), FAssetRegistryTag::TT_Hidden) );
+		Context.AddTag( FAssetRegistryTag(SourceFileTagName(), AssetImportData->GetSourceData().ToJson(), FAssetRegistryTag::TT_Hidden) );
 	}
 
-	Super::GetAssetRegistryTags(OutTags);
+	Super::GetAssetRegistryTags(Context);
 }
 #endif
 
diff --git a/Engine/Plugins/2D/Paper2D/Source/PaperSpriteSheetImporter/Private/PaperSpriteSheet.cpp b/Engine/Plugins/2D/Paper2D/Source/PaperSpriteSheetImporter/Private/PaperSpriteSheet.cpp
--- a/Engine/Plugins/2D/Paper2D/Source/PaperSpriteSheetImporter/Private/PaperSpriteSheet.cpp
+++ b/Engine/Plugins/2D/Paper2D/Source/PaperSpriteSheetImporter/Private/PaperSpriteSheet.cpp
@@ -2,6 +2,7 @@
 
 #include "PaperSpriteSheet.h"
 #include "EditorFramework/AssetImportData.h"
+#include "UObject/AssetRegistryTagsContext.h"
 
 #include UE_INLINE_GENERATED_CPP_BY_NAME(PaperSpriteSheet)
 
@@ -23,13 +24,19 @@ void UPaperSpriteSheet::PostInitProperties()
 
 #if WITH_EDITORONLY_DATA
 void UPaperSpriteSheet::GetAssetRegistryTags(TArray<FAssetRegistryTag>& OutTags) const
+{
+	PRAGMA_DISABLE_DEPRECATION_WARNINGS;
+	Super::GetAssetRegistryTags(OutTags);
+	PRAGMA_ENABLE_DEPRECATION_WARNINGS;
+}
+void UPaperSpriteSheet::GetAssetRegistryTags(FAssetRegistryTagsContext Context) const
 {
 	if (AssetImportData)
 	{
-		OutTags.Add( FAssetRegistryTag(SourceFileTagName(), AssetImportData->GetSourceData().ToJson(), FAssetRegistryTag::TT_Hidden) );
+		Context.AddTag( FAssetRegistryTag(SourceFileTagName(), AssetImportData->GetSourceData().ToJson(), FAssetRegistryTag::TT_Hidden) );
 	}
 
-	Super::GetAssetRegistryTags(OutTags);
+	Super::GetAssetRegistryTags(Context);
 }
 void UPaperSpriteSheet::Serialize(FArchive& Ar)
 {
diff --git a/Engine/Plugins/2D/Paper2D/Source/SpriterImporter/Private/PaperSpriterImportData.cpp b/Engine/Plugins/2D/Paper2D/Source/SpriterImporter/Private/PaperSpriterImportData.cpp
--- a/Engine/Plugins/2D/Paper2D/Source/SpriterImporter/Private/PaperSpriterImportData.cpp
+++ b/Engine/Plugins/2D/Paper2D/Source/SpriterImporter/Private/PaperSpriterImportData.cpp
@@ -1,6 +1,7 @@
 // Copyright Epic Games, Inc. All Rights Reserved.
 
 #include "PaperSpriterImportData.h"
+#include "UObject/AssetRegistryTagsContext.h"
 
 //////////////////////////////////////////////////////////////////////////
 // UPaperSpriterImportData
@@ -12,11 +13,18 @@ UPaperSpriterImportData::UPaperSpriterImportData(const FObjectInitializer& Objec
 }
 
 void UPaperSpriterImportData::GetAssetRegistryTags(TArray<FAssetRegistryTag>& OutTags) const
+{
+	PRAGMA_DISABLE_DEPRECATION_WARNINGS;
+	Super::GetAssetRegistryTags(OutTags);
+	PRAGMA_ENABLE_DEPRECATION_WARNINGS;
+}
+
+void UPaperSpriterImportData::GetAssetRegistryTags(FAssetRegistryTagsContext Context) const
 {
 	if (AssetImportData != nullptr)
 	{
-		OutTags.Add( FAssetRegistryTag(SourceFileTagName(), AssetImportData->GetSourceData().ToJson(), FAssetRegistryTag::TT_Hidden) );
+		Context.AddTag( FAssetRegistryTag(SourceFileTagName(), AssetImportData->GetSourceData().ToJson(), FAssetRegistryTag::TT_Hidden) );
 	}
 
-	Super::GetAssetRegistryTags(OutTags);
+	Super::GetAssetRegistryTags(Context);
 }
diff --git a/Engine/Plugins/Animation/ControlRig/Source/ControlRigEditor/Private/Editor/SRigModuleAssetBrowser.cpp b/Engine/Plugins/Animation/ControlRig/Source/ControlRigEditor/Private/Editor/SRigModuleAssetBrowser.cpp
--- a/Engine/Plugins/Animation/ControlRig/Source/ControlRigEditor/Private/Editor/SRigModuleAssetBrowser.cpp
+++ b/Engine/Plugins/Animation/ControlRig/Source/ControlRigEditor/Private/Editor/SRigModuleAssetBrowser.cpp
@@ -6,6 +6,7 @@
 #include "ControlRigBlueprint.h"
 #include "IContentBrowserSingleton.h"
 #include "Subsystems/AssetEditorSubsystem.h"
+#include "UObject/AssetRegistryTagsContext.h"
 
 #include "ControlRigEditor.h"
 #include "Widgets/Layout/SUniformGridPanel.h"
@@ -52,11 +53,12 @@ void SRigModuleAssetBrowser::RefreshView()
 	AssetPickerConfig.bForceShowEngineContent = true;
 
 	// hide all asset registry columns by default (we only really want the name and path)
-	TArray<UObject::FAssetRegistryTag> AssetRegistryTags;
-	UControlRigBlueprint::StaticClass()->GetDefaultObject()->GetAssetRegistryTags(AssetRegistryTags);
-	for(UObject::FAssetRegistryTag& AssetRegistryTag : AssetRegistryTags)
+	UObject* DefaultControlRigBlueprint = UControlRigBlueprint::StaticClass()->GetDefaultObject();
+	FAssetRegistryTagsContextData Context(DefaultControlRigBlueprint, EAssetRegistryTagsCaller::Uncategorized);
+	DefaultControlRigBlueprint->GetAssetRegistryTags(Context);
+	for (TPair<FName, UObject::FAssetRegistryTag>& AssetRegistryTagPair : Context.Tags)
 	{
-		AssetPickerConfig.HiddenColumnNames.Add(AssetRegistryTag.Name.ToString());
+		AssetPickerConfig.HiddenColumnNames.Add(AssetRegistryTagPair.Value.Name.ToString());
 	}
 
 	// Also hide the type column by default (but allow users to enable it, so don't use bShowTypeInColumnView)
diff --git a/Engine/Plugins/Animation/DeformerGraph/Source/OptimusCore/Private/OptimusDeformer.cpp b/Engine/Plugins/Animation/DeformerGraph/Source/OptimusCore/Private/OptimusDeformer.cpp
--- a/Engine/Plugins/Animation/DeformerGraph/Source/OptimusCore/Private/OptimusDeformer.cpp
+++ b/Engine/Plugins/Animation/DeformerGraph/Source/OptimusCore/Private/OptimusDeformer.cpp
@@ -34,6 +34,7 @@
 #include "RenderingThread.h"
 #include "SceneInterface.h"
 #include "ShaderCore.h"
+#include "UObject/AssetRegistryTagsContext.h"
 #include "UObject/Package.h"
 
 // FIXME: We should not be accessing nodes directly.
@@ -3342,7 +3343,14 @@ void UOptimusDeformer::PostRename(UObject* OldOuter, const FName OldName)
 
 void UOptimusDeformer::GetAssetRegistryTags(TArray<FAssetRegistryTag>& OutTags) const
 {
+	PRAGMA_DISABLE_DEPRECATION_WARNINGS;
 	Super::GetAssetRegistryTags(OutTags);
+	PRAGMA_ENABLE_DEPRECATION_WARNINGS;
+}
+
+void UOptimusDeformer::GetAssetRegistryTags(FAssetRegistryTagsContext Context) const
+{
+	Super::GetAssetRegistryTags(Context);
 
 	UClass* BindingClass = nullptr;
 
@@ -3364,7 +3372,7 @@ void UOptimusDeformer::GetAssetRegistryTags(TArray<FAssetRegistryTag>& OutTags)
 	if (BindingClass != nullptr)
 	{
 		FSoftClassPath ClassPath(BindingClass);
-		OutTags.Add(FAssetRegistryTag(TEXT("PrimaryBindingClass"), *ClassPath.ToString(), FAssetRegistryTag::TT_Hidden));
+		Context.AddTag(FAssetRegistryTag(TEXT("PrimaryBindingClass"), *ClassPath.ToString(), FAssetRegistryTag::TT_Hidden));
 	}
 }
 
diff --git a/Engine/Plugins/Animation/DeformerGraph/Source/OptimusCore/Public/OptimusDeformer.h b/Engine/Plugins/Animation/DeformerGraph/Source/OptimusCore/Public/OptimusDeformer.h
--- a/Engine/Plugins/Animation/DeformerGraph/Source/OptimusCore/Public/OptimusDeformer.h
+++ b/Engine/Plugins/Animation/DeformerGraph/Source/OptimusCore/Public/OptimusDeformer.h
@@ -289,7 +289,9 @@ public:
 	// are not moved to the new package automatically (see FAssetRenameManager), so we
 	// have to manually perform the move/rename, to avoid invalid reference to the old package
 	void PostRename(UObject* OldOuter, const FName OldName) override;
-	void GetAssetRegistryTags(TArray<FAssetRegistryTag>& OutTags) const override;
+	virtual void GetAssetRegistryTags(FAssetRegistryTagsContext Context) const override;
+	UE_DEPRECATED(5.4, "Implement the version that takes FAssetRegistryTagsContext instead.")
+	virtual void GetAssetRegistryTags(TArray<FAssetRegistryTag>& OutTags) const override;
 
 #if WITH_EDITOR
 	void PostEditChangeProperty(FPropertyChangedEvent& PropertyChangedEvent) override;
diff --git a/Engine/Plugins/Animation/IKRig/Source/IKRigEditor/Private/RetargetEditor/IKRetargetPoseExporter.cpp b/Engine/Plugins/Animation/IKRig/Source/IKRigEditor/Private/RetargetEditor/IKRetargetPoseExporter.cpp
--- a/Engine/Plugins/Animation/IKRig/Source/IKRigEditor/Private/RetargetEditor/IKRetargetPoseExporter.cpp
+++ b/Engine/Plugins/Animation/IKRig/Source/IKRigEditor/Private/RetargetEditor/IKRetargetPoseExporter.cpp
@@ -12,6 +12,7 @@
 #include "Framework/Notifications/NotificationManager.h"
 #include "RetargetEditor/IKRetargetAnimInstance.h"
 #include "RetargetEditor/IKRetargetEditorController.h"
+#include "UObject/AssetRegistryTagsContext.h"
 #include "UObject/SavePackage.h"
 #include "Widgets/Input/SComboBox.h"
 #include "Widgets/Notifications/SNotificationList.h"
@@ -311,14 +312,15 @@ void FIKRetargetPoseExporter::HandleImportFromSequenceAsset()
 	AssetPickerConfig.bAllowNullSelection = false;
 
 	// hide all asset registry columns by default (we only really want the name and path)
-	TArray<UObject::FAssetRegistryTag> AssetRegistryTags;
-	UAnimSequence::StaticClass()->GetDefaultObject()->GetAssetRegistryTags(AssetRegistryTags);
+	UObject* AnimSequenceDefaultObject = UAnimSequence::StaticClass()->GetDefaultObject();
+	FAssetRegistryTagsContextData TagsContext(AnimSequenceDefaultObject, EAssetRegistryTagsCaller::Uncategorized);
+	AnimSequenceDefaultObject->GetAssetRegistryTags(TagsContext);
 	FName ColumnToKeep = FName("Number of Frames");
-	for(UObject::FAssetRegistryTag& AssetRegistryTag : AssetRegistryTags)
+	for (const TPair<FName, UObject::FAssetRegistryTag>& TagPair : TagsContext.Tags)
 	{
-		if (AssetRegistryTag.Name != ColumnToKeep)
+		if (TagPair.Key != ColumnToKeep)
 		{
-			AssetPickerConfig.HiddenColumnNames.Add(AssetRegistryTag.Name.ToString());
+			AssetPickerConfig.HiddenColumnNames.Add(TagPair.Key.ToString());
 		}
 	}
 
diff --git a/Engine/Plugins/Animation/IKRig/Source/IKRigEditor/Private/RetargetEditor/SIKRetargetAssetBrowser.cpp b/Engine/Plugins/Animation/IKRig/Source/IKRigEditor/Private/RetargetEditor/SIKRetargetAssetBrowser.cpp
--- a/Engine/Plugins/Animation/IKRig/Source/IKRigEditor/Private/RetargetEditor/SIKRetargetAssetBrowser.cpp
+++ b/Engine/Plugins/Animation/IKRig/Source/IKRigEditor/Private/RetargetEditor/SIKRetargetAssetBrowser.cpp
@@ -16,6 +16,7 @@
 #include "RetargetEditor/SRetargetAnimAssetsWindow.h"
 #include "Retargeter/IKRetargeter.h"
 #include "Retargeter/IKRetargetProcessor.h"
+#include "UObject/AssetRegistryTagsContext.h"
 #include "Widgets/Layout/SUniformGridPanel.h"
 
 #define LOCTEXT_NAMESPACE "IKRetargeterAssetBrowser"
@@ -89,11 +90,12 @@ void SIKRetargetAssetBrowser::RefreshView()
 	AssetPickerConfig.bFocusSearchBoxWhenOpened = false;
 
 	// hide all asset registry columns by default (we only really want the name and path)
-	TArray<UObject::FAssetRegistryTag> AssetRegistryTags;
-	UAnimSequence::StaticClass()->GetDefaultObject()->GetAssetRegistryTags(AssetRegistryTags);
-	for(UObject::FAssetRegistryTag& AssetRegistryTag : AssetRegistryTags)
+	UObject* AnimSequenceDefaultObject = UAnimSequence::StaticClass()->GetDefaultObject();
+	FAssetRegistryTagsContextData TagsContext(AnimSequenceDefaultObject, EAssetRegistryTagsCaller::Uncategorized);
+	AnimSequenceDefaultObject->GetAssetRegistryTags(TagsContext);
+	for (const TPair<FName, UObject::FAssetRegistryTag>& TagPair : TagsContext.Tags)
 	{
-		AssetPickerConfig.HiddenColumnNames.Add(AssetRegistryTag.Name.ToString());
+		AssetPickerConfig.HiddenColumnNames.Add(TagPair.Key.ToString());
 	}
 
 	// Also hide the type column by default (but allow users to enable it, so don't use bShowTypeInColumnView)
diff --git a/Engine/Plugins/Animation/IKRig/Source/IKRigEditor/Private/RigEditor/SIKRigAssetBrowser.cpp b/Engine/Plugins/Animation/IKRig/Source/IKRigEditor/Private/RigEditor/SIKRigAssetBrowser.cpp
--- a/Engine/Plugins/Animation/IKRig/Source/IKRigEditor/Private/RigEditor/SIKRigAssetBrowser.cpp
+++ b/Engine/Plugins/Animation/IKRig/Source/IKRigEditor/Private/RigEditor/SIKRigAssetBrowser.cpp
@@ -11,6 +11,7 @@
 #include "RigEditor/IKRigController.h"
 #include "RigEditor/IKRigEditorController.h"
 #include "Engine/SkeletalMesh.h"
+#include "UObject/AssetRegistryTagsContext.h"
 
 #define LOCTEXT_NAMESPACE "IKRigAssetBrowser"
 
@@ -61,11 +62,12 @@ void SIKRigAssetBrowser::RefreshView()
 	AssetPickerConfig.bFocusSearchBoxWhenOpened = false;
 
 	// hide all asset registry columns by default (we only really want the name and path)
-	TArray<UObject::FAssetRegistryTag> AssetRegistryTags;
-	UAnimSequence::StaticClass()->GetDefaultObject()->GetAssetRegistryTags(AssetRegistryTags);
-	for(UObject::FAssetRegistryTag& AssetRegistryTag : AssetRegistryTags)
+	UObject* AnimSequenceDefaultObject = UAnimSequence::StaticClass()->GetDefaultObject();
+	FAssetRegistryTagsContextData TagsContext(AnimSequenceDefaultObject, EAssetRegistryTagsCaller::Uncategorized);
+	AnimSequenceDefaultObject->GetAssetRegistryTags(TagsContext);
+	for (const TPair<FName, UObject::FAssetRegistryTag>& TagPair : TagsContext.Tags)
 	{
-		AssetPickerConfig.HiddenColumnNames.Add(AssetRegistryTag.Name.ToString());
+		AssetPickerConfig.HiddenColumnNames.Add(TagPair.Key.ToString());
 	}
 
 	// Also hide the type column by default (but allow users to enable it, so don't use bShowTypeInColumnView)
diff --git a/Engine/Plugins/Animation/MLDeformer/MLDeformerFramework/Source/MLDeformerFramework/Private/MLDeformerAsset.cpp b/Engine/Plugins/Animation/MLDeformer/MLDeformerFramework/Source/MLDeformerFramework/Private/MLDeformerAsset.cpp
--- a/Engine/Plugins/Animation/MLDeformer/MLDeformerFramework/Source/MLDeformerFramework/Private/MLDeformerAsset.cpp
+++ b/Engine/Plugins/Animation/MLDeformer/MLDeformerFramework/Source/MLDeformerFramework/Private/MLDeformerAsset.cpp
@@ -5,6 +5,7 @@
 #include "MLDeformerModel.h"
 #include "MLDeformerInputInfo.h"
 #include "Engine/SkeletalMesh.h"
+#include "UObject/AssetRegistryTagsContext.h"
 
 void UMLDeformerAsset::Serialize(FArchive& Archive)
 {
@@ -16,9 +17,16 @@ void UMLDeformerAsset::Serialize(FArchive& Archive)
 
 void UMLDeformerAsset::GetAssetRegistryTags(TArray<FAssetRegistryTag>& OutTags) const
 {
+	PRAGMA_DISABLE_DEPRECATION_WARNINGS;
 	Super::GetAssetRegistryTags(OutTags);
+	PRAGMA_ENABLE_DEPRECATION_WARNINGS;
+}
+
+void UMLDeformerAsset::GetAssetRegistryTags(FAssetRegistryTagsContext Context) const
+{
+	Super::GetAssetRegistryTags(Context);
 	if (Model)
 	{
-		Model->GetAssetRegistryTags(OutTags);
+		Model->GetAssetRegistryTags(Context);
 	}
 }
diff --git a/Engine/Plugins/Animation/MLDeformer/MLDeformerFramework/Source/MLDeformerFramework/Private/MLDeformerGeomCacheModel.cpp b/Engine/Plugins/Animation/MLDeformer/MLDeformerFramework/Source/MLDeformerFramework/Private/MLDeformerGeomCacheModel.cpp
--- a/Engine/Plugins/Animation/MLDeformer/MLDeformerFramework/Source/MLDeformerFramework/Private/MLDeformerGeomCacheModel.cpp
+++ b/Engine/Plugins/Animation/MLDeformer/MLDeformerFramework/Source/MLDeformerFramework/Private/MLDeformerGeomCacheModel.cpp
@@ -5,6 +5,7 @@
 #include "MLDeformerComponent.h"
 #include "MLDeformerAsset.h"
 #include "MLDeformerObjectVersion.h"
+#include "UObject/AssetRegistryTagsContext.h"
 #include "UObject/Object.h"
 #include "UObject/UObjectGlobals.h"
 #include "GeometryCache.h"
@@ -58,7 +59,14 @@ void UMLDeformerGeomCacheModel::PostLoad()
 
 void UMLDeformerGeomCacheModel::GetAssetRegistryTags(TArray<FAssetRegistryTag>& OutTags) const
 {
+	PRAGMA_DISABLE_DEPRECATION_WARNINGS;
 	Super::GetAssetRegistryTags(OutTags);
+	PRAGMA_ENABLE_DEPRECATION_WARNINGS;
+}
+
+void UMLDeformerGeomCacheModel::GetAssetRegistryTags(FAssetRegistryTagsContext Context) const
+{
+	Super::GetAssetRegistryTags(Context);
 
 	#if WITH_EDITORONLY_DATA
 		FString AnimInputString;
@@ -73,8 +81,8 @@ void UMLDeformerGeomCacheModel::GetAssetRegistryTags(TArray<FAssetRegistryTag>&
 		{
 			AnimInputString = TEXT("None");
 		}
-		OutTags.Add(FAssetRegistryTag("MLDeformer.TrainingAnims", AnimInputString, FAssetRegistryTag::TT_Alphabetical));
-		OutTags.Add(FAssetRegistryTag("MLDeformer.NumTrainingAnims", FString::FromInt(TrainingInputAnims.Num()), FAssetRegistryTag::TT_Numerical));
+		Context.AddTag(FAssetRegistryTag("MLDeformer.TrainingAnims", AnimInputString, FAssetRegistryTag::TT_Alphabetical));
+		Context.AddTag(FAssetRegistryTag("MLDeformer.NumTrainingAnims", FString::FromInt(TrainingInputAnims.Num()), FAssetRegistryTag::TT_Numerical));
 	#endif
 }
 
diff --git a/Engine/Plugins/Animation/MLDeformer/MLDeformerFramework/Source/MLDeformerFramework/Private/MLDeformerInputInfo.cpp b/Engine/Plugins/Animation/MLDeformer/MLDeformerFramework/Source/MLDeformerFramework/Private/MLDeformerInputInfo.cpp
--- a/Engine/Plugins/Animation/MLDeformer/MLDeformerFramework/Source/MLDeformerFramework/Private/MLDeformerInputInfo.cpp
+++ b/Engine/Plugins/Animation/MLDeformer/MLDeformerFramework/Source/MLDeformerFramework/Private/MLDeformerInputInfo.cpp
@@ -10,6 +10,7 @@
 #include "ReferenceSkeleton.h"
 #include "Engine/SkeletalMesh.h"
 #include "Components/SkeletalMeshComponent.h"
+#include "UObject/AssetRegistryTagsContext.h"
 
 #include UE_INLINE_GENERATED_CPP_BY_NAME(MLDeformerInputInfo)
 void UMLDeformerInputInfo::Reset()
@@ -23,12 +24,19 @@ void UMLDeformerInputInfo::Reset()
 
 void UMLDeformerInputInfo::GetAssetRegistryTags(TArray<FAssetRegistryTag>& OutTags) const
 {
+	PRAGMA_DISABLE_DEPRECATION_WARNINGS;
 	Super::GetAssetRegistryTags(OutTags);
-	OutTags.Add(FAssetRegistryTag("MLDeformer.Trained.SkeletalMesh", SkeletalMesh.IsValid() ? SkeletalMesh.ToString() : TEXT("None"), FAssetRegistryTag::TT_Alphabetical));
-	OutTags.Add(FAssetRegistryTag("MLDeformer.Trained.NumBaseMeshVertices", FString::FromInt(NumBaseMeshVertices), FAssetRegistryTag::TT_Numerical));
-	OutTags.Add(FAssetRegistryTag("MLDeformer.Trained.NumTargetMeshVertices", FString::FromInt(NumTargetMeshVertices), FAssetRegistryTag::TT_Numerical));
-	OutTags.Add(FAssetRegistryTag("MLDeformer.Trained.NumBones", FString::FromInt(GetNumBones()), FAssetRegistryTag::TT_Numerical));
-	OutTags.Add(FAssetRegistryTag("MLDeformer.Trained.NumCurves", FString::FromInt(GetNumCurves()), FAssetRegistryTag::TT_Numerical));
+	PRAGMA_ENABLE_DEPRECATION_WARNINGS;
+}
+
+void UMLDeformerInputInfo::GetAssetRegistryTags(FAssetRegistryTagsContext Context) const
+{
+	Super::GetAssetRegistryTags(Context);
+	Context.AddTag(FAssetRegistryTag("MLDeformer.Trained.SkeletalMesh", SkeletalMesh.IsValid() ? SkeletalMesh.ToString() : TEXT("None"), FAssetRegistryTag::TT_Alphabetical));
+	Context.AddTag(FAssetRegistryTag("MLDeformer.Trained.NumBaseMeshVertices", FString::FromInt(NumBaseMeshVertices), FAssetRegistryTag::TT_Numerical));
+	Context.AddTag(FAssetRegistryTag("MLDeformer.Trained.NumTargetMeshVertices", FString::FromInt(NumTargetMeshVertices), FAssetRegistryTag::TT_Numerical));
+	Context.AddTag(FAssetRegistryTag("MLDeformer.Trained.NumBones", FString::FromInt(GetNumBones()), FAssetRegistryTag::TT_Numerical));
+	Context.AddTag(FAssetRegistryTag("MLDeformer.Trained.NumCurves", FString::FromInt(GetNumCurves()), FAssetRegistryTag::TT_Numerical));
 }
 
 void UMLDeformerInputInfo::OnPostLoad()
diff --git a/Engine/Plugins/Animation/MLDeformer/MLDeformerFramework/Source/MLDeformerFramework/Private/MLDeformerModel.cpp b/Engine/Plugins/Animation/MLDeformer/MLDeformerFramework/Source/MLDeformerFramework/Private/MLDeformerModel.cpp
--- a/Engine/Plugins/Animation/MLDeformer/MLDeformerFramework/Source/MLDeformerFramework/Private/MLDeformerModel.cpp
+++ b/Engine/Plugins/Animation/MLDeformer/MLDeformerFramework/Source/MLDeformerFramework/Private/MLDeformerModel.cpp
@@ -16,6 +16,7 @@
 #include "RHICommandList.h"
 #include "AssetRegistry/AssetData.h"
 #include "Components/SkeletalMeshComponent.h"
+#include "UObject/AssetRegistryTagsContext.h"
 
 #include UE_INLINE_GENERATED_CPP_BY_NAME(MLDeformerModel)
 
@@ -158,25 +159,32 @@ bool UMLDeformerModel::IsReadyForFinishDestroy()
 
 void UMLDeformerModel::GetAssetRegistryTags(TArray<FAssetRegistryTag>& OutTags) const
 {
+	PRAGMA_DISABLE_DEPRECATION_WARNINGS;
 	Super::GetAssetRegistryTags(OutTags);
+	PRAGMA_ENABLE_DEPRECATION_WARNINGS;
+}
+
+void UMLDeformerModel::GetAssetRegistryTags(FAssetRegistryTagsContext Context) const
+{
+	Super::GetAssetRegistryTags(Context);
 
-	OutTags.Add(FAssetRegistryTag("MLDeformer.ModelType", GetClass()->GetName(), FAssetRegistryTag::TT_Alphabetical));
-	OutTags.Add(FAssetRegistryTag("MLDeformer.IsTrained", IsTrained() ? TEXT("True") : TEXT("False"), FAssetRegistryTag::TT_Alphabetical));
-	OutTags.Add(FAssetRegistryTag("MLDeformer.NumBaseMeshVerts", FString::FromInt(NumBaseMeshVerts), FAssetRegistryTag::TT_Numerical));
-	OutTags.Add(FAssetRegistryTag("MLDeformer.NumTargetMeshVerts", FString::FromInt(NumTargetMeshVerts), FAssetRegistryTag::TT_Numerical));
-	OutTags.Add(FAssetRegistryTag("MLDeformer.SkeletalMesh", SkeletalMesh ? FAssetData(SkeletalMesh).ToSoftObjectPath().ToString() : TEXT("None"), FAssetRegistryTag::TT_Alphabetical));
-	OutTags.Add(FAssetRegistryTag("MLDeformer.MaxNumLODs", FString::FromInt(GetMaxNumLODs()), FAssetRegistryTag::TT_Numerical));
+	Context.AddTag(FAssetRegistryTag("MLDeformer.ModelType", GetClass()->GetName(), FAssetRegistryTag::TT_Alphabetical));
+	Context.AddTag(FAssetRegistryTag("MLDeformer.IsTrained", IsTrained() ? TEXT("True") : TEXT("False"), FAssetRegistryTag::TT_Alphabetical));
+	Context.AddTag(FAssetRegistryTag("MLDeformer.NumBaseMeshVerts", FString::FromInt(NumBaseMeshVerts), FAssetRegistryTag::TT_Numerical));
+	Context.AddTag(FAssetRegistryTag("MLDeformer.NumTargetMeshVerts", FString::FromInt(NumTargetMeshVerts), FAssetRegistryTag::TT_Numerical));
+	Context.AddTag(FAssetRegistryTag("MLDeformer.SkeletalMesh", SkeletalMesh ? FAssetData(SkeletalMesh).ToSoftObjectPath().ToString() : TEXT("None"), FAssetRegistryTag::TT_Alphabetical));
+	Context.AddTag(FAssetRegistryTag("MLDeformer.MaxNumLODs", FString::FromInt(GetMaxNumLODs()), FAssetRegistryTag::TT_Numerical));
 
 	#if WITH_EDITORONLY_DATA
-		OutTags.Add(FAssetRegistryTag("MLDeformer.NumBones", FString::FromInt(BoneIncludeList.Num()), FAssetRegistryTag::TT_Numerical));
-		OutTags.Add(FAssetRegistryTag("MLDeformer.NumCurves", FString::FromInt(CurveIncludeList.Num()), FAssetRegistryTag::TT_Numerical));
-		OutTags.Add(FAssetRegistryTag("MLDeformer.DeltaCutoffLength", FString::Printf(TEXT("%f"), DeltaCutoffLength), FAssetRegistryTag::TT_Numerical));
-		OutTags.Add(FAssetRegistryTag("MLDeformer.MaxTrainingFrames", FString::FromInt(MaxTrainingFrames), FAssetRegistryTag::TT_Numerical));
+		Context.AddTag(FAssetRegistryTag("MLDeformer.NumBones", FString::FromInt(BoneIncludeList.Num()), FAssetRegistryTag::TT_Numerical));
+		Context.AddTag(FAssetRegistryTag("MLDeformer.NumCurves", FString::FromInt(CurveIncludeList.Num()), FAssetRegistryTag::TT_Numerical));
+		Context.AddTag(FAssetRegistryTag("MLDeformer.DeltaCutoffLength", FString::Printf(TEXT("%f"), DeltaCutoffLength), FAssetRegistryTag::TT_Numerical));
+		Context.AddTag(FAssetRegistryTag("MLDeformer.MaxTrainingFrames", FString::FromInt(MaxTrainingFrames), FAssetRegistryTag::TT_Numerical));
 	#endif
 
 	if (InputInfo)
 	{
-		InputInfo->GetAssetRegistryTags(OutTags);
+		InputInfo->GetAssetRegistryTags(Context);
 	}
 }
 
diff --git a/Engine/Plugins/Animation/MLDeformer/MLDeformerFramework/Source/MLDeformerFramework/Private/MLDeformerMorphModel.cpp b/Engine/Plugins/Animation/MLDeformer/MLDeformerFramework/Source/MLDeformerFramework/Private/MLDeformerMorphModel.cpp
--- a/Engine/Plugins/Animation/MLDeformer/MLDeformerFramework/Source/MLDeformerFramework/Private/MLDeformerMorphModel.cpp
+++ b/Engine/Plugins/Animation/MLDeformer/MLDeformerFramework/Source/MLDeformerFramework/Private/MLDeformerMorphModel.cpp
@@ -10,6 +10,7 @@
 #include "Components/ExternalMorphSet.h"
 #include "Components/SkeletalMeshComponent.h"
 #include "Rendering/MorphTargetVertexInfoBuffers.h"
+#include "UObject/AssetRegistryTagsContext.h"
 
 #include UE_INLINE_GENERATED_CPP_BY_NAME(MLDeformerMorphModel)
 
@@ -128,17 +129,24 @@ void UMLDeformerMorphModel::PostLoad()
 
 void UMLDeformerMorphModel::GetAssetRegistryTags(TArray<FAssetRegistryTag>& OutTags) const
 {
+	PRAGMA_DISABLE_DEPRECATION_WARNINGS;
 	Super::GetAssetRegistryTags(OutTags);
+	PRAGMA_ENABLE_DEPRECATION_WARNINGS;
+}
 
-	OutTags.Add(FAssetRegistryTag("MLDeformer.MorphModel.ClampMorphWeights", bClampMorphWeights ? TEXT("True") : TEXT("False"), FAssetRegistryTag::TT_Alphabetical));
-	OutTags.Add(FAssetRegistryTag("MLDeformer.MorphModel.InvertMaskChannel", bInvertMaskChannel? TEXT("True") : TEXT("False"), FAssetRegistryTag::TT_Alphabetical));
-	OutTags.Add(FAssetRegistryTag("MLDeformer.MorphModel.IncludeNormals", bIncludeNormals ? TEXT("True") : TEXT("False"), FAssetRegistryTag::TT_Alphabetical));
-	OutTags.Add(FAssetRegistryTag("MLDeformer.MorphModel.NumMorphTargets", FString::FromInt(GetNumMorphTargets(0)), FAssetRegistryTag::TT_Numerical));
-	OutTags.Add(FAssetRegistryTag("MLDeformer.MorphModel.DeltaZeroThreshold", FString::Printf(TEXT("%f"), MorphDeltaZeroThreshold), FAssetRegistryTag::TT_Numerical));
-	OutTags.Add(FAssetRegistryTag("MLDeformer.MorphModel.CompressionLevel", FString::Printf(TEXT("%f"), MorphCompressionLevel), FAssetRegistryTag::TT_Numerical));
-	OutTags.Add(FAssetRegistryTag("MLDeformer.MorphModel.CompressedSize", FString::FromInt(CompressedMorphDataSizeInBytes), FAssetRegistryTag::TT_Numerical));
-	OutTags.Add(FAssetRegistryTag("MLDeformer.MorphModel.UncompressedSize", FString::FromInt(UncompressedMorphDataSizeInBytes), FAssetRegistryTag::TT_Numerical));
-	OutTags.Add(FAssetRegistryTag("MLDeformer.MorphModel.NumLODs", FString::FromInt(GetNumLODs()), FAssetRegistryTag::TT_Numerical));
+void UMLDeformerMorphModel::GetAssetRegistryTags(FAssetRegistryTagsContext Context) const
+{
+	Super::GetAssetRegistryTags(Context);
+
+	Context.AddTag(FAssetRegistryTag("MLDeformer.MorphModel.ClampMorphWeights", bClampMorphWeights ? TEXT("True") : TEXT("False"), FAssetRegistryTag::TT_Alphabetical));
+	Context.AddTag(FAssetRegistryTag("MLDeformer.MorphModel.InvertMaskChannel", bInvertMaskChannel? TEXT("True") : TEXT("False"), FAssetRegistryTag::TT_Alphabetical));
+	Context.AddTag(FAssetRegistryTag("MLDeformer.MorphModel.IncludeNormals", bIncludeNormals ? TEXT("True") : TEXT("False"), FAssetRegistryTag::TT_Alphabetical));
+	Context.AddTag(FAssetRegistryTag("MLDeformer.MorphModel.NumMorphTargets", FString::FromInt(GetNumMorphTargets(0)), FAssetRegistryTag::TT_Numerical));
+	Context.AddTag(FAssetRegistryTag("MLDeformer.MorphModel.DeltaZeroThreshold", FString::Printf(TEXT("%f"), MorphDeltaZeroThreshold), FAssetRegistryTag::TT_Numerical));
+	Context.AddTag(FAssetRegistryTag("MLDeformer.MorphModel.CompressionLevel", FString::Printf(TEXT("%f"), MorphCompressionLevel), FAssetRegistryTag::TT_Numerical));
+	Context.AddTag(FAssetRegistryTag("MLDeformer.MorphModel.CompressedSize", FString::FromInt(CompressedMorphDataSizeInBytes), FAssetRegistryTag::TT_Numerical));
+	Context.AddTag(FAssetRegistryTag("MLDeformer.MorphModel.UncompressedSize", FString::FromInt(UncompressedMorphDataSizeInBytes), FAssetRegistryTag::TT_Numerical));
+	Context.AddTag(FAssetRegistryTag("MLDeformer.MorphModel.NumLODs", FString::FromInt(GetNumLODs()), FAssetRegistryTag::TT_Numerical));
 }
 
 int32 UMLDeformerMorphModel::GetNumMorphTargets(int32 LOD) const
diff --git a/Engine/Plugins/Animation/MLDeformer/NearestNeighborModel/Source/NearestNeighborModel/Private/NearestNeighborModel.cpp b/Engine/Plugins/Animation/MLDeformer/NearestNeighborModel/Source/NearestNeighborModel/Private/NearestNeighborModel.cpp
--- a/Engine/Plugins/Animation/MLDeformer/NearestNeighborModel/Source/NearestNeighborModel/Private/NearestNeighborModel.cpp
+++ b/Engine/Plugins/Animation/MLDeformer/NearestNeighborModel/Source/NearestNeighborModel/Private/NearestNeighborModel.cpp
@@ -22,6 +22,7 @@
 #include "Rendering/MorphTargetVertexInfoBuffers.h"
 #include "Rendering/SkeletalMeshModel.h"
 #include "ShaderCore.h"
+#include "UObject/AssetRegistryTagsContext.h"
 
 #include UE_INLINE_GENERATED_CPP_BY_NAME(NearestNeighborModel)
 
@@ -1015,21 +1016,28 @@ UNearestNeighborModel::UNearestNeighborModel(const FObjectInitializer& ObjectIni
 
 void UNearestNeighborModel::GetAssetRegistryTags(TArray<FAssetRegistryTag>& OutTags) const
 {
+	PRAGMA_DISABLE_DEPRECATION_WARNINGS;
 	Super::GetAssetRegistryTags(OutTags);
+	PRAGMA_ENABLE_DEPRECATION_WARNINGS;
+}
+
+void UNearestNeighborModel::GetAssetRegistryTags(FAssetRegistryTagsContext Context) const
+{
+	Super::GetAssetRegistryTags(Context);
 
 	#if WITH_EDITORONLY_DATA
-		OutTags.Add(FAssetRegistryTag("MLDeformer.NearestNeighborModel.NumEpochs", FString::FromInt(NumEpochs), FAssetRegistryTag::TT_Numerical));
-		OutTags.Add(FAssetRegistryTag("MLDeformer.NearestNeighborModel.BatchSize", FString::FromInt(BatchSize), FAssetRegistryTag::TT_Numerical));
-		OutTags.Add(FAssetRegistryTag("MLDeformer.NearestNeighborModel.NumHiddenLayers", FString::FromInt(HiddenLayerDims.Num()), FAssetRegistryTag::TT_Numerical));
-		OutTags.Add(FAssetRegistryTag("MLDeformer.NearestNeighborModel.LearningRate", FString::Printf(TEXT("%f"), LearningRate), FAssetRegistryTag::TT_Numerical));
-		OutTags.Add(FAssetRegistryTag("MLDeformer.NearestNeighborModel.EarlyStopEpochs", FString::FromInt(EarlyStopEpochs), FAssetRegistryTag::TT_Numerical));
-		OutTags.Add(FAssetRegistryTag("MLDeformer.NearestNeighborModel.UseDualQuaternions", bUseDualQuaternionDeltas ? TEXT("True") : TEXT("False"), FAssetRegistryTag::TT_Alphabetical));
+		Context.AddTag(FAssetRegistryTag("MLDeformer.NearestNeighborModel.NumEpochs", FString::FromInt(NumEpochs), FAssetRegistryTag::TT_Numerical));
+		Context.AddTag(FAssetRegistryTag("MLDeformer.NearestNeighborModel.BatchSize", FString::FromInt(BatchSize), FAssetRegistryTag::TT_Numerical));
+		Context.AddTag(FAssetRegistryTag("MLDeformer.NearestNeighborModel.NumHiddenLayers", FString::FromInt(HiddenLayerDims.Num()), FAssetRegistryTag::TT_Numerical));
+		Context.AddTag(FAssetRegistryTag("MLDeformer.NearestNeighborModel.LearningRate", FString::Printf(TEXT("%f"), LearningRate), FAssetRegistryTag::TT_Numerical));
+		Context.AddTag(FAssetRegistryTag("MLDeformer.NearestNeighborModel.EarlyStopEpochs", FString::FromInt(EarlyStopEpochs), FAssetRegistryTag::TT_Numerical));
+		Context.AddTag(FAssetRegistryTag("MLDeformer.NearestNeighborModel.UseDualQuaternions", bUseDualQuaternionDeltas ? TEXT("True") : TEXT("False"), FAssetRegistryTag::TT_Alphabetical));
 	#endif
 
-	OutTags.Add(FAssetRegistryTag("MLDeformer.NearestNeighborModel.DecayFactor", FString::Printf(TEXT("%f"), DecayFactor), FAssetRegistryTag::TT_Numerical));
-	OutTags.Add(FAssetRegistryTag("MLDeformer.NearestNeighborModel.NearestNeighborOffsetWeight", FString::Printf(TEXT("%f"), NearestNeighborOffsetWeight), FAssetRegistryTag::TT_Numerical));
-	OutTags.Add(FAssetRegistryTag("MLDeformer.NearestNeighborModel.UseRBF", bUseRBF ? TEXT("True") : TEXT("False"), FAssetRegistryTag::TT_Alphabetical));
-	OutTags.Add(FAssetRegistryTag("MLDeformer.NearestNeighborModel.RBFSigma", FString::Printf(TEXT("%f"), RBFSigma), FAssetRegistryTag::TT_Numerical));
+	Context.AddTag(FAssetRegistryTag("MLDeformer.NearestNeighborModel.DecayFactor", FString::Printf(TEXT("%f"), DecayFactor), FAssetRegistryTag::TT_Numerical));
+	Context.AddTag(FAssetRegistryTag("MLDeformer.NearestNeighborModel.NearestNeighborOffsetWeight", FString::Printf(TEXT("%f"), NearestNeighborOffsetWeight), FAssetRegistryTag::TT_Numerical));
+	Context.AddTag(FAssetRegistryTag("MLDeformer.NearestNeighborModel.UseRBF", bUseRBF ? TEXT("True") : TEXT("False"), FAssetRegistryTag::TT_Alphabetical));
+	Context.AddTag(FAssetRegistryTag("MLDeformer.NearestNeighborModel.RBFSigma", FString::Printf(TEXT("%f"), RBFSigma), FAssetRegistryTag::TT_Numerical));
 }
 
 #if WITH_EDITOR
diff --git a/Engine/Plugins/Animation/MLDeformer/NeuralMorphModel/Source/NeuralMorphModel/Private/NeuralMorphModel.cpp b/Engine/Plugins/Animation/MLDeformer/NeuralMorphModel/Source/NeuralMorphModel/Private/NeuralMorphModel.cpp
--- a/Engine/Plugins/Animation/MLDeformer/NeuralMorphModel/Source/NeuralMorphModel/Private/NeuralMorphModel.cpp
+++ b/Engine/Plugins/Animation/MLDeformer/NeuralMorphModel/Source/NeuralMorphModel/Private/NeuralMorphModel.cpp
@@ -7,6 +7,7 @@
 #include "MLDeformerAsset.h"
 #include "MLDeformerComponent.h"
 #include "MLDeformerObjectVersion.h"
+#include "UObject/AssetRegistryTagsContext.h"
 #include "UObject/UObjectGlobals.h"
 
 #include UE_INLINE_GENERATED_CPP_BY_NAME(NeuralMorphModel)
@@ -137,28 +138,35 @@ bool UNeuralMorphModel::IsTrained() const
 
 void UNeuralMorphModel::GetAssetRegistryTags(TArray<FAssetRegistryTag>& OutTags) const
 {
+	PRAGMA_DISABLE_DEPRECATION_WARNINGS;
 	Super::GetAssetRegistryTags(OutTags);
+	PRAGMA_ENABLE_DEPRECATION_WARNINGS;
+}
+
+void UNeuralMorphModel::GetAssetRegistryTags(FAssetRegistryTagsContext Context) const
+{
+	Super::GetAssetRegistryTags(Context);
 
 	const UNeuralMorphInputInfo* NeuralInputInfo = Cast<UNeuralMorphInputInfo>(GetInputInfo());
 	if (NeuralInputInfo)
 	{
-		OutTags.Add(FAssetRegistryTag("MLDeformer.NeuralMorphModel.Trained.NumBoneGroups", FString::FromInt(NeuralInputInfo->GetBoneGroups().Num()), FAssetRegistryTag::TT_Numerical));
-		OutTags.Add(FAssetRegistryTag("MLDeformer.NeuralMorphModel.Trained.NumCurveGroups", FString::FromInt(NeuralInputInfo->GetCurveGroups().Num()), FAssetRegistryTag::TT_Numerical));
+		Context.AddTag(FAssetRegistryTag("MLDeformer.NeuralMorphModel.Trained.NumBoneGroups", FString::FromInt(NeuralInputInfo->GetBoneGroups().Num()), FAssetRegistryTag::TT_Numerical));
+		Context.AddTag(FAssetRegistryTag("MLDeformer.NeuralMorphModel.Trained.NumCurveGroups", FString::FromInt(NeuralInputInfo->GetCurveGroups().Num()), FAssetRegistryTag::TT_Numerical));
 	}
 
-	OutTags.Add(FAssetRegistryTag("MLDeformer.NeuralMorphModel.Mode", Mode == ENeuralMorphMode::Local ? TEXT("Local") : TEXT("Global"), FAssetRegistryTag::TT_Alphabetical));
-	OutTags.Add(FAssetRegistryTag("MLDeformer.NeuralMorphModel.NumMorphsPerBone", FString::FromInt(LocalNumMorphTargetsPerBone), FAssetRegistryTag::TT_Numerical));
-	OutTags.Add(FAssetRegistryTag("MLDeformer.NeuralMorphModel.NumMorphTargets", FString::FromInt(GlobalNumMorphTargets), FAssetRegistryTag::TT_Numerical));
-	OutTags.Add(FAssetRegistryTag("MLDeformer.NeuralMorphModel.NumIterations", FString::FromInt(NumIterations), FAssetRegistryTag::TT_Numerical));
-	OutTags.Add(FAssetRegistryTag("MLDeformer.NeuralMorphModel.LocalNumHiddenLayers", FString::FromInt(LocalNumHiddenLayers), FAssetRegistryTag::TT_Numerical));
-	OutTags.Add(FAssetRegistryTag("MLDeformer.NeuralMorphModel.LocalNumNeuronsPerLayer", FString::FromInt(LocalNumNeuronsPerLayer), FAssetRegistryTag::TT_Numerical));
-	OutTags.Add(FAssetRegistryTag("MLDeformer.NeuralMorphModel.GlobalNumHiddenLayers", FString::FromInt(GlobalNumHiddenLayers), FAssetRegistryTag::TT_Numerical));
-	OutTags.Add(FAssetRegistryTag("MLDeformer.NeuralMorphModel.GlobalNumNeuronsPerLayer", FString::FromInt(GlobalNumNeuronsPerLayer), FAssetRegistryTag::TT_Numerical));
-	OutTags.Add(FAssetRegistryTag("MLDeformer.NeuralMorphModel.BatchSize", FString::FromInt(BatchSize), FAssetRegistryTag::TT_Numerical));
-	OutTags.Add(FAssetRegistryTag("MLDeformer.NeuralMorphModel.LearningRate", FString::Printf(TEXT("%f"), LearningRate), FAssetRegistryTag::TT_Numerical));
-	OutTags.Add(FAssetRegistryTag("MLDeformer.NeuralMorphModel.RegularizationFactor", FString::Printf(TEXT("%f"), RegularizationFactor), FAssetRegistryTag::TT_Numerical));
-	OutTags.Add(FAssetRegistryTag("MLDeformer.NeuralMorphModel.SmoothLossBeta", FString::Printf(TEXT("%f"), SmoothLossBeta), FAssetRegistryTag::TT_Numerical));
-	OutTags.Add(FAssetRegistryTag("MLDeformer.NeuralMorphModel.EnableBoneMasks", bEnableBoneMasks ? TEXT("True") : TEXT("False"), FAssetRegistryTag::TT_Alphabetical));
+	Context.AddTag(FAssetRegistryTag("MLDeformer.NeuralMorphModel.Mode", Mode == ENeuralMorphMode::Local ? TEXT("Local") : TEXT("Global"), FAssetRegistryTag::TT_Alphabetical));
+	Context.AddTag(FAssetRegistryTag("MLDeformer.NeuralMorphModel.NumMorphsPerBone", FString::FromInt(LocalNumMorphTargetsPerBone), FAssetRegistryTag::TT_Numerical));
+	Context.AddTag(FAssetRegistryTag("MLDeformer.NeuralMorphModel.NumMorphTargets", FString::FromInt(GlobalNumMorphTargets), FAssetRegistryTag::TT_Numerical));
+	Context.AddTag(FAssetRegistryTag("MLDeformer.NeuralMorphModel.NumIterations", FString::FromInt(NumIterations), FAssetRegistryTag::TT_Numerical));
+	Context.AddTag(FAssetRegistryTag("MLDeformer.NeuralMorphModel.LocalNumHiddenLayers", FString::FromInt(LocalNumHiddenLayers), FAssetRegistryTag::TT_Numerical));
+	Context.AddTag(FAssetRegistryTag("MLDeformer.NeuralMorphModel.LocalNumNeuronsPerLayer", FString::FromInt(LocalNumNeuronsPerLayer), FAssetRegistryTag::TT_Numerical));
+	Context.AddTag(FAssetRegistryTag("MLDeformer.NeuralMorphModel.GlobalNumHiddenLayers", FString::FromInt(GlobalNumHiddenLayers), FAssetRegistryTag::TT_Numerical));
+	Context.AddTag(FAssetRegistryTag("MLDeformer.NeuralMorphModel.GlobalNumNeuronsPerLayer", FString::FromInt(GlobalNumNeuronsPerLayer), FAssetRegistryTag::TT_Numerical));
+	Context.AddTag(FAssetRegistryTag("MLDeformer.NeuralMorphModel.BatchSize", FString::FromInt(BatchSize), FAssetRegistryTag::TT_Numerical));
+	Context.AddTag(FAssetRegistryTag("MLDeformer.NeuralMorphModel.LearningRate", FString::Printf(TEXT("%f"), LearningRate), FAssetRegistryTag::TT_Numerical));
+	Context.AddTag(FAssetRegistryTag("MLDeformer.NeuralMorphModel.RegularizationFactor", FString::Printf(TEXT("%f"), RegularizationFactor), FAssetRegistryTag::TT_Numerical));
+	Context.AddTag(FAssetRegistryTag("MLDeformer.NeuralMorphModel.SmoothLossBeta", FString::Printf(TEXT("%f"), SmoothLossBeta), FAssetRegistryTag::TT_Numerical));
+	Context.AddTag(FAssetRegistryTag("MLDeformer.NeuralMorphModel.EnableBoneMasks", bEnableBoneMasks ? TEXT("True") : TEXT("False"), FAssetRegistryTag::TT_Alphabetical));
 }
 
 #undef LOCTEXT_NAMESPACE
diff --git a/Engine/Plugins/Animation/MLDeformer/VertexDeltaModel/Source/VertexDeltaModel/Private/VertexDeltaModel.cpp b/Engine/Plugins/Animation/MLDeformer/VertexDeltaModel/Source/VertexDeltaModel/Private/VertexDeltaModel.cpp
--- a/Engine/Plugins/Animation/MLDeformer/VertexDeltaModel/Source/VertexDeltaModel/Private/VertexDeltaModel.cpp
+++ b/Engine/Plugins/Animation/MLDeformer/VertexDeltaModel/Source/VertexDeltaModel/Private/VertexDeltaModel.cpp
@@ -13,6 +13,7 @@
 #include "Interfaces/IPluginManager.h"
 #include "Misc/Paths.h"
 #include "ShaderCore.h"
+#include "UObject/AssetRegistryTagsContext.h"
 
 #define LOCTEXT_NAMESPACE "VertexDeltaModel"
 
@@ -65,14 +66,21 @@ bool UVertexDeltaModel::IsTrained() const
 
 void UVertexDeltaModel::GetAssetRegistryTags(TArray<FAssetRegistryTag>& OutTags) const
 {
+	PRAGMA_DISABLE_DEPRECATION_WARNINGS;
 	Super::GetAssetRegistryTags(OutTags);
+	PRAGMA_ENABLE_DEPRECATION_WARNINGS;
+}
+
+void UVertexDeltaModel::GetAssetRegistryTags(FAssetRegistryTagsContext Context) const
+{
+	Super::GetAssetRegistryTags(Context);
 
 	#if WITH_EDITORONLY_DATA
-		OutTags.Add(FAssetRegistryTag("MLDeformer.VertexDeltaModel.NumHiddenLayers", FString::FromInt(NumHiddenLayers), FAssetRegistryTag::TT_Numerical));
-		OutTags.Add(FAssetRegistryTag("MLDeformer.VertexDeltaModel.NumNeuronsPerLayer", FString::FromInt(NumNeuronsPerLayer), FAssetRegistryTag::TT_Numerical));
-		OutTags.Add(FAssetRegistryTag("MLDeformer.VertexDeltaModel.NumIterations", FString::FromInt(NumIterations), FAssetRegistryTag::TT_Numerical));
-		OutTags.Add(FAssetRegistryTag("MLDeformer.VertexDeltaModel.BatchSize", FString::FromInt(BatchSize), FAssetRegistryTag::TT_Numerical));
-		OutTags.Add(FAssetRegistryTag("MLDeformer.VertexDeltaModel.LearningRate", FString::Printf(TEXT("%f"), LearningRate), FAssetRegistryTag::TT_Numerical));
+		Context.AddTag(FAssetRegistryTag("MLDeformer.VertexDeltaModel.NumHiddenLayers", FString::FromInt(NumHiddenLayers), FAssetRegistryTag::TT_Numerical));
+		Context.AddTag(FAssetRegistryTag("MLDeformer.VertexDeltaModel.NumNeuronsPerLayer", FString::FromInt(NumNeuronsPerLayer), FAssetRegistryTag::TT_Numerical));
+		Context.AddTag(FAssetRegistryTag("MLDeformer.VertexDeltaModel.NumIterations", FString::FromInt(NumIterations), FAssetRegistryTag::TT_Numerical));
+		Context.AddTag(FAssetRegistryTag("MLDeformer.VertexDeltaModel.BatchSize", FString::FromInt(BatchSize), FAssetRegistryTag::TT_Numerical));
+		Context.AddTag(FAssetRegistryTag("MLDeformer.VertexDeltaModel.LearningRate", FString::Printf(TEXT("%f"), LearningRate), FAssetRegistryTag::TT_Numerical));
 	#endif
 }
 
diff --git a/Engine/Plugins/Animation/RigLogic/Source/RigLogicEditor/Private/RigLogicEditor.cpp b/Engine/Plugins/Animation/RigLogic/Source/RigLogicEditor/Private/RigLogicEditor.cpp
--- a/Engine/Plugins/Animation/RigLogic/Source/RigLogicEditor/Private/RigLogicEditor.cpp
+++ b/Engine/Plugins/Animation/RigLogic/Source/RigLogicEditor/Private/RigLogicEditor.cpp
@@ -18,6 +18,7 @@
 #include "Misc/Paths.h"
 #include "AssetRegistry/IAssetRegistry.h"
 #include "AssetRegistry/AssetData.h"
+#include "UObject/AssetRegistryTagsContext.h"
 
 
 IMPLEMENT_MODULE(FRigLogicEditor, RigLogicEditor)
@@ -33,7 +34,7 @@ void FRigLogicEditor::StartupModule()
 
 	MenuExtenderDelegates.Add(FContentBrowserMenuExtender_SelectedAssets::CreateStatic(&FRigLogicEditor::OnExtendSkelMeshWithDNASelectionMenu));
 
-	UObject::FAssetRegistryTag::OnGetExtraObjectTags.AddStatic(&GetAssetRegistryTagsForDNA);
+	UObject::FAssetRegistryTag::OnGetExtraObjectTagsWithContext.AddStatic(&GetAssetRegistryTagsForDNA);
 }
 
 
@@ -146,8 +147,9 @@ void FRigLogicEditor::ExecuteDNAReimport(class UObject* Mesh)
 	}
 }
 
-void FRigLogicEditor::GetAssetRegistryTagsForDNA(const class UObject* Object, TArray<UObject::FAssetRegistryTag>& OutTags)
+void FRigLogicEditor::GetAssetRegistryTagsForDNA(FAssetRegistryTagsContext Context)
 {
+	const UObject* Object = Context.GetObject();
 	if ((Object != nullptr) && (Object->GetClass() != nullptr) && Object->GetClass()->IsChildOf(USkeletalMesh::StaticClass()))
 	{
 		FString DNAname = (LOCTEXT("DnaNotOnSkeletalMesh", "No DNA Attached")).ToString();
@@ -166,7 +168,7 @@ void FRigLogicEditor::GetAssetRegistryTagsForDNA(const class UObject* Object, TA
 				}
 			}
 		}
-		OutTags.Add(UObject::FAssetRegistryTag("DNA", DNAname, UObject::FAssetRegistryTag::TT_Alphabetical));
+		Context.AddTag(UObject::FAssetRegistryTag("DNA", DNAname, UObject::FAssetRegistryTag::TT_Alphabetical));
 	}
 }
 
diff --git a/Engine/Plugins/Animation/RigLogic/Source/RigLogicEditor/Public/RigLogicEditor.h b/Engine/Plugins/Animation/RigLogic/Source/RigLogicEditor/Public/RigLogicEditor.h
--- a/Engine/Plugins/Animation/RigLogic/Source/RigLogicEditor/Public/RigLogicEditor.h
+++ b/Engine/Plugins/Animation/RigLogic/Source/RigLogicEditor/Public/RigLogicEditor.h
@@ -26,5 +26,5 @@ private:
 	static void ExecuteDNAImport(UObject* Mesh);
 	static void ExecuteDNAReimport(UObject* Mesh);
 
-	static void GetAssetRegistryTagsForDNA(const UObject* Object, TArray<UObject::FAssetRegistryTag>& OutTags);
+	static void GetAssetRegistryTagsForDNA(FAssetRegistryTagsContext Context);
 };
diff --git a/Engine/Plugins/Compositing/OpenColorIO/Source/OpenColorIO/Private/OpenColorIOConfiguration.cpp b/Engine/Plugins/Compositing/OpenColorIO/Source/OpenColorIO/Private/OpenColorIOConfiguration.cpp
--- a/Engine/Plugins/Compositing/OpenColorIO/Source/OpenColorIO/Private/OpenColorIOConfiguration.cpp
+++ b/Engine/Plugins/Compositing/OpenColorIO/Source/OpenColorIO/Private/OpenColorIOConfiguration.cpp
@@ -17,6 +17,7 @@
 #include "TextureResource.h"
 #include "Widgets/Notifications/SNotificationList.h"
 
+#include "UObject/AssetRegistryTagsContext.h"
 #include "UObject/ObjectSaveContext.h"
 
 #include UE_INLINE_GENERATED_CPP_BY_NAME(OpenColorIOConfiguration)
@@ -553,7 +554,14 @@ void UOpenColorIOConfiguration::PostLoad()
 
 void UOpenColorIOConfiguration::GetAssetRegistryTags(TArray<FAssetRegistryTag>& OutTags) const
 {
+	PRAGMA_DISABLE_DEPRECATION_WARNINGS;
 	Super::GetAssetRegistryTags(OutTags);
+	PRAGMA_ENABLE_DEPRECATION_WARNINGS;
+}
+
+void UOpenColorIOConfiguration::GetAssetRegistryTags(FAssetRegistryTagsContext Context) const
+{
+	Super::GetAssetRegistryTags(Context);
 
 	FString Description;
 	if (ConfigurationFile.FilePath.IsEmpty())
@@ -569,7 +577,7 @@ void UOpenColorIOConfiguration::GetAssetRegistryTags(TArray<FAssetRegistryTag>&
 		Description = TEXT("Configuration: ") + ConfigurationFile.FilePath;
 	}
 
-	OutTags.Add(FAssetRegistryTag(TEXT("ConfigurationFile"), Description, FAssetRegistryTag::TT_Hidden));
+	Context.AddTag(FAssetRegistryTag(TEXT("ConfigurationFile"), Description, FAssetRegistryTag::TT_Hidden));
 }
 
 #if WITH_EDITORONLY_DATA
diff --git a/Engine/Plugins/Editor/AssetSearch/Source/Private/AssetSearchManager.cpp b/Engine/Plugins/Editor/AssetSearch/Source/Private/AssetSearchManager.cpp
--- a/Engine/Plugins/Editor/AssetSearch/Source/Private/AssetSearchManager.cpp
+++ b/Engine/Plugins/Editor/AssetSearch/Source/Private/AssetSearchManager.cpp
@@ -486,7 +486,7 @@ uint64 FAssetSearchManager::GetTextHash(FStringView PackageRelativeExportPath) c
 	return CityHash64(reinterpret_cast<const char*>(PackageRelativeExportPath.GetData() + 1), (PackageRelativeExportPath.Len() - 1) * sizeof(TCHAR));
 }
 
-void FAssetSearchManager::HandleOnGetExtraObjectTags(FAssetRegistryTagsContext& Context)
+void FAssetSearchManager::HandleOnGetExtraObjectTags(FAssetRegistryTagsContext Context)
 {
 	if (!Context.IsFullUpdate())
 	{
diff --git a/Engine/Plugins/Editor/AssetSearch/Source/Private/AssetSearchManager.h b/Engine/Plugins/Editor/AssetSearch/Source/Private/AssetSearchManager.h
--- a/Engine/Plugins/Editor/AssetSearch/Source/Private/AssetSearchManager.h
+++ b/Engine/Plugins/Editor/AssetSearch/Source/Private/AssetSearchManager.h
@@ -57,7 +57,7 @@ private:
 	void OnAssetRemoved(const FAssetData& InAssetData);
 	void OnAssetScanFinished();
 
-	void HandleOnGetExtraObjectTags(FAssetRegistryTagsContext& Context);
+	void HandleOnGetExtraObjectTags(FAssetRegistryTagsContext Context);
 	void HandlePackageSaved(const FString& PackageFilename, UPackage* Package, FObjectPostSaveContext ObjectSaveContext);
 	void OnAssetLoaded(UObject* InObject);
 
diff --git a/Engine/Plugins/Editor/ContentBrowser/ContentBrowserAssetDataSource/Source/ContentBrowserAssetDataSource/Private/AssetPropertyTagCache.cpp b/Engine/Plugins/Editor/ContentBrowser/ContentBrowserAssetDataSource/Source/ContentBrowserAssetDataSource/Private/AssetPropertyTagCache.cpp
--- a/Engine/Plugins/Editor/ContentBrowser/ContentBrowserAssetDataSource/Source/ContentBrowserAssetDataSource/Private/AssetPropertyTagCache.cpp
+++ b/Engine/Plugins/Editor/ContentBrowser/ContentBrowserAssetDataSource/Source/ContentBrowserAssetDataSource/Private/AssetPropertyTagCache.cpp
@@ -1,6 +1,7 @@
 // Copyright Epic Games, Inc. All Rights Reserved.
 
 #include "AssetPropertyTagCache.h"
+#include "UObject/AssetRegistryTagsContext.h"
 #include "UObject/UnrealType.h"
 #include "UObject/LinkerLoad.h"
 
@@ -37,19 +38,20 @@ const FAssetPropertyTagCache::FClassPropertyTagCache& FAssetPropertyTagCache::Ge
 		if (UClass* AssetClass = GetAssetClass())
 		{
 			// Get the tags data from the CDO and use it to build the cache
-			TArray<UObject::FAssetRegistryTag> AssetTags;
-			AssetClass->GetDefaultObject()->GetAssetRegistryTags(AssetTags);
+			UObject* AssetClassDefaultObject = AssetClass->GetDefaultObject();
+			FAssetRegistryTagsContextData TagsContext(AssetClassDefaultObject, EAssetRegistryTagsCaller::Uncategorized);
+			AssetClassDefaultObject->GetAssetRegistryTags(TagsContext);
 			TMap<FName, UObject::FAssetRegistryTagMetadata> AssetTagMetaData;
-			AssetClass->GetDefaultObject()->GetAssetRegistryTagMetadata(AssetTagMetaData);
+			AssetClassDefaultObject->GetAssetRegistryTagMetadata(AssetTagMetaData);
 
-			ClassCache->TagNameToCachedDataMap.Reserve(AssetTags.Num());
-			for (const UObject::FAssetRegistryTag& AssetTag : AssetTags)
+			ClassCache->TagNameToCachedDataMap.Reserve(TagsContext.Tags.Num());
+			for (const TPair<FName, UObject::FAssetRegistryTag>& TagPair : TagsContext.Tags)
 			{
-				FPropertyTagCache& TagCache = ClassCache->TagNameToCachedDataMap.Add(AssetTag.Name);
-				TagCache.TagType = AssetTag.Type;
-				TagCache.DisplayFlags = AssetTag.DisplayFlags;
+				FPropertyTagCache& TagCache = ClassCache->TagNameToCachedDataMap.Add(TagPair.Key);
+				TagCache.TagType = TagPair.Value.Type;
+				TagCache.DisplayFlags = TagPair.Value.DisplayFlags;
 
-				if (const UObject::FAssetRegistryTagMetadata* TagMetaData = AssetTagMetaData.Find(AssetTag.Name))
+				if (const UObject::FAssetRegistryTagMetadata* TagMetaData = AssetTagMetaData.Find(TagPair.Key))
 				{
 					TagCache.DisplayName = TagMetaData->DisplayName;
 					TagCache.TooltipText = TagMetaData->TooltipText;
@@ -59,14 +61,14 @@ const FAssetPropertyTagCache::FClassPropertyTagCache& FAssetPropertyTagCache::Ge
 				else
 				{
 					// If the tag name corresponds to a property name, use the property tooltip
-					const FProperty* Property = FindFProperty<FProperty>(AssetClass, AssetTag.Name);
-					TagCache.TooltipText = Property ? Property->GetToolTipText() : FText::FromString(FName::NameToDisplayString(AssetTag.Name.ToString(), false));
+					const FProperty* Property = FindFProperty<FProperty>(AssetClass, TagPair.Key);
+					TagCache.TooltipText = Property ? Property->GetToolTipText() : FText::FromString(FName::NameToDisplayString(TagPair.Key.ToString(), false));
 				}
 
 				// Ensure a display name for this tag
 				if (TagCache.DisplayName.IsEmpty())
 				{
-					if (const FProperty* TagField = FindFProperty<FProperty>(AssetClass, AssetTag.Name))
+					if (const FProperty* TagField = FindFProperty<FProperty>(AssetClass, TagPair.Key))
 					{
 						// Take the display name from the corresponding property if possible
 						TagCache.DisplayName = TagField->GetDisplayNameText();
@@ -74,7 +76,7 @@ const FAssetPropertyTagCache::FClassPropertyTagCache& FAssetPropertyTagCache::Ge
 					else
 					{
 						// We have no type information by this point, so no idea if it's a bool :(
-						TagCache.DisplayName = FText::AsCultureInvariant(FName::NameToDisplayString(AssetTag.Name.ToString(), /*bIsBool*/false));
+						TagCache.DisplayName = FText::AsCultureInvariant(FName::NameToDisplayString(TagPair.Key.ToString(), /*bIsBool*/false));
 					}
 				}
 
@@ -82,9 +84,9 @@ const FAssetPropertyTagCache::FClassPropertyTagCache& FAssetPropertyTagCache::Ge
 				// This is useful as people only see the display name in the UI, so are more likely to use it
 				{
 					const FName SanitizedDisplayName = MakeObjectNameFromDisplayLabel(TagCache.DisplayName.ToString(), NAME_None);
-					if (AssetTag.Name != SanitizedDisplayName)
+					if (TagPair.Key != SanitizedDisplayName)
 					{
-						ClassCache->DisplayNameToTagNameMap.Add(SanitizedDisplayName, AssetTag.Name);
+						ClassCache->DisplayNameToTagNameMap.Add(SanitizedDisplayName, TagPair.Key);
 					}
 				}
 			}
diff --git a/Engine/Plugins/Editor/ContentBrowser/ContentBrowserClassDataSource/Source/ContentBrowserClassDataSource/Private/ContentBrowserClassDataCore.cpp b/Engine/Plugins/Editor/ContentBrowser/ContentBrowserClassDataSource/Source/ContentBrowserClassDataSource/Private/ContentBrowserClassDataCore.cpp
--- a/Engine/Plugins/Editor/ContentBrowser/ContentBrowserClassDataSource/Source/ContentBrowserClassDataSource/Private/ContentBrowserClassDataCore.cpp
+++ b/Engine/Plugins/Editor/ContentBrowser/ContentBrowserClassDataSource/Source/ContentBrowserClassDataSource/Private/ContentBrowserClassDataCore.cpp
@@ -8,6 +8,7 @@
 #include "AssetViewUtils.h"
 #include "IAssetTypeActions.h"
 #include "ContentBrowserDataUtils.h"
+#include "UObject/AssetRegistryTagsContext.h"
 
 #define LOCTEXT_NAMESPACE "ContentBrowserClassDataSource"
 
@@ -347,15 +348,16 @@ const FClassTagDefinitionMap& GetAvailableClassTags()
 	{
 		FClassTagDefinitionMap ClassTagsTmp;
 
-		TArray<UObject::FAssetRegistryTag> CDOClassTags;
-		GetDefault<UClass>()->GetAssetRegistryTags(CDOClassTags);
+		const UObject* ClassDefault = GetDefault<UClass>();
+		FAssetRegistryTagsContextData TagsContext(ClassDefault, EAssetRegistryTagsCaller::Uncategorized);
+		ClassDefault->GetAssetRegistryTags(TagsContext);
 
-		for (const UObject::FAssetRegistryTag& CDOClassTag : CDOClassTags)
+		for (const TPair<FName,UObject::FAssetRegistryTag>& TagPair : TagsContext.Tags)
 		{
-			FClassTagDefinition& ClassTag = ClassTagsTmp.Add(CDOClassTag.Name);
-			ClassTag.TagType = CDOClassTag.Type;
-			ClassTag.DisplayFlags = CDOClassTag.DisplayFlags;
-			ClassTag.DisplayName = FText::AsCultureInvariant(FName::NameToDisplayString(CDOClassTag.Name.ToString(), /*bIsBool*/false));
+			FClassTagDefinition& ClassTag = ClassTagsTmp.Add(TagPair.Key);
+			ClassTag.TagType = TagPair.Value.Type;
+			ClassTag.DisplayFlags = TagPair.Value.DisplayFlags;
+			ClassTag.DisplayName = FText::AsCultureInvariant(FName::NameToDisplayString(TagPair.Key.ToString(), /*bIsBool*/false));
 		}
 
 		return ClassTagsTmp;
diff --git a/Engine/Plugins/Enterprise/DatasmithContent/Source/DatasmithContent/Private/DatasmithAssetImportData.cpp b/Engine/Plugins/Enterprise/DatasmithContent/Source/DatasmithContent/Private/DatasmithAssetImportData.cpp
--- a/Engine/Plugins/Enterprise/DatasmithContent/Source/DatasmithContent/Private/DatasmithAssetImportData.cpp
+++ b/Engine/Plugins/Enterprise/DatasmithContent/Source/DatasmithContent/Private/DatasmithAssetImportData.cpp
@@ -12,15 +12,26 @@
 #include "Materials/MaterialInterface.h"
 #include "Misc/FileHelper.h"
 #include "Misc/Paths.h"
+#include "UObject/AssetRegistryTagsContext.h"
 #include "UObject/EnterpriseObjectVersion.h"
 
 #if WITH_EDITORONLY_DATA
 void FDatasmithImportInfo::GetAssetRegistryTags(TArray<UObject::FAssetRegistryTag>& OutTags) const
+{
+	FAssetRegistryTagsContextData TagsContext(nullptr, EAssetRegistryTagsCaller::Uncategorized);
+	GetAssetRegistryTags(TagsContext);
+	for (TPair<FName, UObject::FAssetRegistryTag>& Pair : TagsContext.Tags)
+	{
+		OutTags.Add(MoveTemp(Pair.Value));
+	}
+}
+
+void FDatasmithImportInfo::GetAssetRegistryTags(FAssetRegistryTagsContext Context) const
 {
 	using FAssetRegistryTag = UObject::FAssetRegistryTag;
 
-	OutTags.Add(FAssetRegistryTag(TEXT("SourceUri"), SourceUri, FAssetRegistryTag::TT_Hidden));
-	OutTags.Add(FAssetRegistryTag(TEXT("SourceHash"), SourceHash, FAssetRegistryTag::TT_Hidden));
+	Context.AddTag(FAssetRegistryTag(TEXT("SourceUri"), SourceUri, FAssetRegistryTag::TT_Hidden));
+	Context.AddTag(FAssetRegistryTag(TEXT("SourceHash"), SourceHash, FAssetRegistryTag::TT_Hidden));
 }
 
 #if WITH_EDITOR
diff --git a/Engine/Plugins/Enterprise/DatasmithContent/Source/DatasmithContent/Private/DatasmithScene.cpp b/Engine/Plugins/Enterprise/DatasmithContent/Source/DatasmithContent/Private/DatasmithScene.cpp
--- a/Engine/Plugins/Enterprise/DatasmithContent/Source/DatasmithContent/Private/DatasmithScene.cpp
+++ b/Engine/Plugins/Enterprise/DatasmithContent/Source/DatasmithContent/Private/DatasmithScene.cpp
@@ -8,6 +8,7 @@
 #include "Engine/World.h"
 #include "EngineUtils.h"
 #include "LevelSequence.h"
+#include "UObject/AssetRegistryTagsContext.h"
 #include "UObject/EnterpriseObjectVersion.h"
 #include "UObject/Package.h"
 
@@ -40,12 +41,20 @@ UDatasmithScene::~UDatasmithScene()
 
 void UDatasmithScene::GetAssetRegistryTags(TArray<FAssetRegistryTag>& OutTags) const
 {
+	PRAGMA_DISABLE_DEPRECATION_WARNINGS;
+	Super::GetAssetRegistryTags(OutTags);
+	PRAGMA_ENABLE_DEPRECATION_WARNINGS;
+}
+
+void UDatasmithScene::GetAssetRegistryTags(FAssetRegistryTagsContext Context) const
+{
+	Super::GetAssetRegistryTags(Context);
 #if WITH_EDITORONLY_DATA
 	if (AssetImportData)
 	{
-		OutTags.Add(FAssetRegistryTag(SourceFileTagName(), AssetImportData->GetSourceData().ToJson(), FAssetRegistryTag::TT_Hidden));
+		Context.AddTag(FAssetRegistryTag(SourceFileTagName(), AssetImportData->GetSourceData().ToJson(), FAssetRegistryTag::TT_Hidden));
 
-		AssetImportData->DatasmithImportInfo.GetAssetRegistryTags(OutTags);
+		AssetImportData->DatasmithImportInfo.GetAssetRegistryTags(Context);
 	}
 #endif
 }
diff --git a/Engine/Plugins/Enterprise/DatasmithContent/Source/DatasmithContent/Public/DatasmithAssetImportData.h b/Engine/Plugins/Enterprise/DatasmithContent/Source/DatasmithContent/Public/DatasmithAssetImportData.h
--- a/Engine/Plugins/Enterprise/DatasmithContent/Source/DatasmithContent/Public/DatasmithAssetImportData.h
+++ b/Engine/Plugins/Enterprise/DatasmithContent/Source/DatasmithContent/Public/DatasmithAssetImportData.h
@@ -42,6 +42,8 @@ struct DATASMITHCONTENT_API FDatasmithImportInfo
 	UPROPERTY()
 	FString SourceHash;
 
+	void GetAssetRegistryTags(FAssetRegistryTagsContext Context) const;
+	UE_DEPRECATED(5.4, "Call the version that takes FAssetRegistryTagsContext instead.")
 	void GetAssetRegistryTags(TArray<UObject::FAssetRegistryTag>& OutTags) const;
 #endif // WITH_EDITORONLY_DATA
 };
diff --git a/Engine/Plugins/Enterprise/LidarPointCloud/Source/LidarPointCloudRuntime/Private/LidarPointCloud.cpp b/Engine/Plugins/Enterprise/LidarPointCloud/Source/LidarPointCloudRuntime/Private/LidarPointCloud.cpp
--- a/Engine/Plugins/Enterprise/LidarPointCloud/Source/LidarPointCloudRuntime/Private/LidarPointCloud.cpp
+++ b/Engine/Plugins/Enterprise/LidarPointCloud/Source/LidarPointCloudRuntime/Private/LidarPointCloud.cpp
@@ -17,6 +17,7 @@
 #include "Framework/Notifications/NotificationManager.h"
 #include "EngineUtils.h"
 #include "Components/BrushComponent.h"
+#include "UObject/AssetRegistryTagsContext.h"
 #include "UObject/GCObjectScopeGuard.h"
 #include "Widgets/Notifications/SNotificationList.h"
 #include "UObject/ObjectSaveContext.h"
@@ -369,10 +370,17 @@ void ULidarPointCloud::PostLoad()
 
 void ULidarPointCloud::GetAssetRegistryTags(TArray<FAssetRegistryTag>& OutTags) const
 {
-	OutTags.Add(FAssetRegistryTag("PointCount", PointCloudAssetRegistryCache.PointCount, FAssetRegistryTag::TT_Numerical));
-	OutTags.Add(FAssetRegistryTag("ApproxSize", PointCloudAssetRegistryCache.ApproxSize, FAssetRegistryTag::TT_Dimensional));
-
+	PRAGMA_DISABLE_DEPRECATION_WARNINGS;
 	Super::GetAssetRegistryTags(OutTags);
+	PRAGMA_ENABLE_DEPRECATION_WARNINGS;
+}
+
+void ULidarPointCloud::GetAssetRegistryTags(FAssetRegistryTagsContext Context) const
+{
+	Context.AddTag(FAssetRegistryTag("PointCount", PointCloudAssetRegistryCache.PointCount, FAssetRegistryTag::TT_Numerical));
+	Context.AddTag(FAssetRegistryTag("ApproxSize", PointCloudAssetRegistryCache.ApproxSize, FAssetRegistryTag::TT_Dimensional));
+
+	Super::GetAssetRegistryTags(Context);
 }
 
 void ULidarPointCloud::BeginDestroy()
diff --git a/Engine/Plugins/Experimental/AnimNext/Source/AnimNext/Internal/Scheduler/AnimNextSchedule.h b/Engine/Plugins/Experimental/AnimNext/Source/AnimNext/Internal/Scheduler/AnimNextSchedule.h
--- a/Engine/Plugins/Experimental/AnimNext/Source/AnimNext/Internal/Scheduler/AnimNextSchedule.h
+++ b/Engine/Plugins/Experimental/AnimNext/Source/AnimNext/Internal/Scheduler/AnimNextSchedule.h
@@ -236,6 +236,8 @@ private:
 #if WITH_EDITOR
 	virtual void PostEditChangeProperty(FPropertyChangedEvent& PropertyChangedEvent) override;
 	virtual void PostEditUndo() override;
+	virtual void GetAssetRegistryTags(FAssetRegistryTagsContext Context) const override;
+	UE_DEPRECATED(5.4, "Implement the version that takes FAssetRegistryTagsContext instead.")
 	virtual void GetAssetRegistryTags(TArray<FAssetRegistryTag>& OutTags) const override;
 
 	// Compile the editor data into a compact runtime representation
@@ -247,7 +249,7 @@ private:
 	static TUniqueFunction<void(UAnimNextSchedule*)> CompileFunction;
 
 	// Function hook used to expose to asset registry
-	static TUniqueFunction<void(const UAnimNextSchedule*, TArray<FAssetRegistryTag>&)> GetAssetRegistryTagsFunction;
+	static TUniqueFunction<void(const UAnimNextSchedule*, FAssetRegistryTagsContext)> GetAssetRegistryTagsFunction;
 	
 	// Editor only
 	// TODO: move this into an editor only subobject
diff --git a/Engine/Plugins/Experimental/AnimNext/Source/AnimNext/Private/Graph/AnimNextGraph.cpp b/Engine/Plugins/Experimental/AnimNext/Source/AnimNext/Private/Graph/AnimNextGraph.cpp
--- a/Engine/Plugins/Experimental/AnimNext/Source/AnimNext/Private/Graph/AnimNextGraph.cpp
+++ b/Engine/Plugins/Experimental/AnimNext/Source/AnimNext/Private/Graph/AnimNextGraph.cpp
@@ -4,6 +4,7 @@
 #include "RigVMCore/RigVMMemoryStorage.h"
 #include "RigVMCore/RigVMExecuteContext.h"
 #include "Graph/RigUnit_AnimNextShimRoot.h"
+#include "UObject/AssetRegistryTagsContext.h"
 #include "UObject/ObjectSaveContext.h"
 #include "UObject/Package.h"
 #include "Param/ParamTypeHandle.h"
@@ -91,11 +92,18 @@ void UAnimNextGraph::AllocateInstanceImpl(FAnimNextGraphInstance* ParentGraphIns
 #if WITH_EDITOR
 void UAnimNextGraph::GetAssetRegistryTags(TArray<FAssetRegistryTag>& OutTags) const
 {
+	PRAGMA_DISABLE_DEPRECATION_WARNINGS;
 	Super::GetAssetRegistryTags(OutTags);
+	PRAGMA_ENABLE_DEPRECATION_WARNINGS;
+}
+
+void UAnimNextGraph::GetAssetRegistryTags(FAssetRegistryTagsContext Context) const
+{
+	Super::GetAssetRegistryTags(Context);
 
 	if (EditorData)
 	{		
-		EditorData->GetAssetRegistryTags(OutTags);
+		EditorData->GetAssetRegistryTags(Context);
 	}
 }
 #endif
diff --git a/Engine/Plugins/Experimental/AnimNext/Source/AnimNext/Private/Param/AnimNextParameterBlock.cpp b/Engine/Plugins/Experimental/AnimNext/Source/AnimNext/Private/Param/AnimNextParameterBlock.cpp
--- a/Engine/Plugins/Experimental/AnimNext/Source/AnimNext/Private/Param/AnimNextParameterBlock.cpp
+++ b/Engine/Plugins/Experimental/AnimNext/Source/AnimNext/Private/Param/AnimNextParameterBlock.cpp
@@ -3,6 +3,7 @@
 #include "Param/AnimNextParameterBlock.h"
 #include "RigVMCore/RigVMMemoryStorage.h"
 #include "RigVMRuntimeDataRegistry.h"
+#include "UObject/AssetRegistryTagsContext.h"
 #include "UObject/ObjectSaveContext.h"
 #include "UObject/Package.h"
 #include "Graph/RigUnit_AnimNextBeginExecution.h"
@@ -74,12 +75,19 @@ void UAnimNextParameterBlock::PostLoad()
 
 void UAnimNextParameterBlock::GetAssetRegistryTags(TArray<FAssetRegistryTag>& OutTags) const
 {
+	PRAGMA_DISABLE_DEPRECATION_WARNINGS;
 	Super::GetAssetRegistryTags(OutTags);
+	PRAGMA_ENABLE_DEPRECATION_WARNINGS;
+}
+
+void UAnimNextParameterBlock::GetAssetRegistryTags(FAssetRegistryTagsContext Context) const
+{
+	Super::GetAssetRegistryTags(Context);
 	
 #if WITH_EDITORONLY_DATA
 	if(EditorData)
 	{
-		EditorData->GetAssetRegistryTags(OutTags);
+		EditorData->GetAssetRegistryTags(Context);
 	}
 #endif
 }
\ No newline at end of file
diff --git a/Engine/Plugins/Experimental/AnimNext/Source/AnimNext/Private/Scheduler/AnimNextSchedule.cpp b/Engine/Plugins/Experimental/AnimNext/Source/AnimNext/Private/Scheduler/AnimNextSchedule.cpp
--- a/Engine/Plugins/Experimental/AnimNext/Source/AnimNext/Private/Scheduler/AnimNextSchedule.cpp
+++ b/Engine/Plugins/Experimental/AnimNext/Source/AnimNext/Private/Scheduler/AnimNextSchedule.cpp
@@ -4,10 +4,11 @@
 #include "Tasks/Task.h"
 #include "Async/TaskGraphInterfaces.h"
 #include "EngineLogs.h"
+#include "UObject/AssetRegistryTagsContext.h"
 
 #if WITH_EDITOR
 TUniqueFunction<void(UAnimNextSchedule*)> UAnimNextSchedule::CompileFunction;
-TUniqueFunction<void(const UAnimNextSchedule*, TArray<UObject::FAssetRegistryTag>&)> UAnimNextSchedule::GetAssetRegistryTagsFunction;
+TUniqueFunction<void(const UAnimNextSchedule*, FAssetRegistryTagsContext)> UAnimNextSchedule::GetAssetRegistryTagsFunction;
 #endif
 
 void UAnimNextSchedule::PostLoad()
@@ -44,11 +45,18 @@ void UAnimNextSchedule::CompileSchedule()
 
 void UAnimNextSchedule::GetAssetRegistryTags(TArray<FAssetRegistryTag>& OutTags) const
 {
+	PRAGMA_DISABLE_DEPRECATION_WARNINGS;
 	Super::GetAssetRegistryTags(OutTags);
+	PRAGMA_ENABLE_DEPRECATION_WARNINGS;
+}
+
+void UAnimNextSchedule::GetAssetRegistryTags(FAssetRegistryTagsContext Context) const
+{
+	Super::GetAssetRegistryTags(Context);
 
 	check(GetAssetRegistryTagsFunction);
 
-	GetAssetRegistryTagsFunction(this, OutTags);
+	GetAssetRegistryTagsFunction(this, Context);
 }
 
 #endif // #if WITH_EDITOR
\ No newline at end of file
diff --git a/Engine/Plugins/Experimental/AnimNext/Source/AnimNextEditor/Private/Workspace/AnimNextWorkspace.cpp b/Engine/Plugins/Experimental/AnimNext/Source/AnimNextEditor/Private/Workspace/AnimNextWorkspace.cpp
--- a/Engine/Plugins/Experimental/AnimNext/Source/AnimNextEditor/Private/Workspace/AnimNextWorkspace.cpp
+++ b/Engine/Plugins/Experimental/AnimNext/Source/AnimNextEditor/Private/Workspace/AnimNextWorkspace.cpp
@@ -1,10 +1,11 @@
-﻿// Copyright Epic Games, Inc. All Rights Reserved.
+// Copyright Epic Games, Inc. All Rights Reserved.
 
 #include "AnimNextWorkspace.h"
 
 #include "Graph/AnimNextGraph.h"
 #include "Param/AnimNextParameterBlock.h"
 #include "Scheduler/AnimNextSchedule.h"
+#include "UObject/AssetRegistryTagsContext.h"
 
 const FName UAnimNextWorkspace::ExportsAssetRegistryTag = TEXT("Exports");
 
@@ -214,7 +215,14 @@ void UAnimNextWorkspace::PostTransacted(const FTransactionObjectEvent& Transacti
 
 void UAnimNextWorkspace::GetAssetRegistryTags(TArray<FAssetRegistryTag>& OutTags) const
 {
+	PRAGMA_DISABLE_DEPRECATION_WARNINGS;
 	Super::GetAssetRegistryTags(OutTags);
+	PRAGMA_ENABLE_DEPRECATION_WARNINGS;
+}
+
+void UAnimNextWorkspace::GetAssetRegistryTags(FAssetRegistryTagsContext Context) const
+{
+	Super::GetAssetRegistryTags(Context);
 
 	FAnimNextWorkspaceAssetRegistryExports Exports;
 	Exports.Assets.Reserve(Assets.Num());
@@ -227,5 +235,5 @@ void UAnimNextWorkspace::GetAssetRegistryTags(TArray<FAssetRegistryTag>& OutTags
 	FString TagValue;
 	FAnimNextWorkspaceAssetRegistryExports::StaticStruct()->ExportText(TagValue, &Exports, nullptr, nullptr, PPF_None, nullptr);
 
-	OutTags.Add(FAssetRegistryTag(ExportsAssetRegistryTag, TagValue, FAssetRegistryTag::TT_Hidden));
+	Context.AddTag(FAssetRegistryTag(ExportsAssetRegistryTag, TagValue, FAssetRegistryTag::TT_Hidden));
 }
diff --git a/Engine/Plugins/Experimental/AnimNext/Source/AnimNextEditor/Private/Workspace/AnimNextWorkspace.h b/Engine/Plugins/Experimental/AnimNext/Source/AnimNextEditor/Private/Workspace/AnimNextWorkspace.h
--- a/Engine/Plugins/Experimental/AnimNext/Source/AnimNextEditor/Private/Workspace/AnimNextWorkspace.h
+++ b/Engine/Plugins/Experimental/AnimNext/Source/AnimNextEditor/Private/Workspace/AnimNextWorkspace.h
@@ -1,4 +1,4 @@
-﻿// Copyright Epic Games, Inc. All Rights Reserved.
+// Copyright Epic Games, Inc. All Rights Reserved.
 
 #pragma once
 
@@ -91,6 +91,8 @@ class UAnimNextWorkspace : public UObject
 	void BroadcastModified();
 
 	// UObject interface
+	virtual void GetAssetRegistryTags(FAssetRegistryTagsContext Context) const override;
+	UE_DEPRECATED(5.4, "Implement the version that takes FAssetRegistryTagsContext instead.")
 	virtual void GetAssetRegistryTags(TArray<FAssetRegistryTag>& OutTags) const override;
 	virtual void PostTransacted(const FTransactionObjectEvent& TransactionEvent) override;
 	virtual bool IsEditorOnly() const override { return true; }
diff --git a/Engine/Plugins/Experimental/AnimNext/Source/AnimNextUncookedOnly/Private/AnimNextUncookedOnlyModule.cpp b/Engine/Plugins/Experimental/AnimNext/Source/AnimNextUncookedOnly/Private/AnimNextUncookedOnlyModule.cpp
--- a/Engine/Plugins/Experimental/AnimNext/Source/AnimNextUncookedOnly/Private/AnimNextUncookedOnlyModule.cpp
+++ b/Engine/Plugins/Experimental/AnimNext/Source/AnimNextUncookedOnly/Private/AnimNextUncookedOnlyModule.cpp
@@ -9,6 +9,7 @@
 #include "Modules/ModuleManager.h"
 #include "Param/AnimNextTag.h"
 #include "Scheduler/AnimNextSchedule.h"
+#include "UObject/AssetRegistryTagsContext.h"
 
 namespace UE::AnimNext::UncookedOnly
 {
@@ -24,19 +25,20 @@ namespace UE::AnimNext::UncookedOnly
 			};
 
 			// TEMP: Bind the asset registry tags function for schedules
-			UAnimNextSchedule::GetAssetRegistryTagsFunction = [](const UAnimNextSchedule* InSchedule, TArray<UObject::FAssetRegistryTag>& OutTags)
+			UAnimNextSchedule::GetAssetRegistryTagsFunction = [](const UAnimNextSchedule* InSchedule, FAssetRegistryTagsContext Context)
 			{
 				FAnimNextParameterProviderAssetRegistryExports Exports;
 				FUtils::GetScheduleParameters(InSchedule, Exports);
 				
 				FString TagValue;
 				FAnimNextParameterProviderAssetRegistryExports::StaticStruct()->ExportText(TagValue, &Exports, nullptr, nullptr, PPF_None, nullptr);
-				OutTags.Add(UObject::FAssetRegistryTag(UE::AnimNext::ExportsAnimNextAssetRegistryTag, TagValue, UObject::FAssetRegistryTag::TT_Hidden));
+				Context.AddTag(UObject::FAssetRegistryTag(UE::AnimNext::ExportsAnimNextAssetRegistryTag, TagValue, UObject::FAssetRegistryTag::TT_Hidden));
 			};
 			
 			// Ensure that any BP components that we care about contribute to the parameter pool
-			OnGetExtraObjectTagsHandle = UObject::FAssetRegistryTag::OnGetExtraObjectTags.AddLambda([](const UObject* InObject, TArray<UObject::FAssetRegistryTag>& OutTags)
+			OnGetExtraObjectTagsHandle = UObject::FAssetRegistryTag::OnGetExtraObjectTagsWithContext.AddLambda([](FAssetRegistryTagsContext Context)
 			{
+				const UObject* InObject = Context.GetObject();
 				if(const UBlueprint* Blueprint = Cast<UBlueprint>(InObject))
 				{
 					FAnimNextParameterProviderAssetRegistryExports Exports;
@@ -44,7 +46,7 @@ namespace UE::AnimNext::UncookedOnly
 
 					FString TagValue;
 					FAnimNextParameterProviderAssetRegistryExports::StaticStruct()->ExportText(TagValue, &Exports, nullptr, nullptr, PPF_None, nullptr);
-					OutTags.Add(UObject::FAssetRegistryTag(UE::AnimNext::ExportsAnimNextAssetRegistryTag, TagValue, UObject::FAssetRegistryTag::TT_Hidden));
+					Context.AddTag(UObject::FAssetRegistryTag(UE::AnimNext::ExportsAnimNextAssetRegistryTag, TagValue, UObject::FAssetRegistryTag::TT_Hidden));
 				}
 			});
 		}
@@ -54,7 +56,7 @@ namespace UE::AnimNext::UncookedOnly
 			UAnimNextSchedule::GetAssetRegistryTagsFunction = nullptr;
 			UAnimNextSchedule::CompileFunction = nullptr;
 
-			UObject::FAssetRegistryTag::OnGetExtraObjectTags.Remove(OnGetExtraObjectTagsHandle);
+			UObject::FAssetRegistryTag::OnGetExtraObjectTagsWithContext.Remove(OnGetExtraObjectTagsHandle);
 		}
 
 		FDelegateHandle OnGetExtraObjectTagsHandle;
diff --git a/Engine/Plugins/Experimental/AnimNext/Source/AnimNextUncookedOnly/Private/Graph/AnimNextGraph_EditorData.cpp b/Engine/Plugins/Experimental/AnimNext/Source/AnimNextUncookedOnly/Private/Graph/AnimNextGraph_EditorData.cpp
--- a/Engine/Plugins/Experimental/AnimNext/Source/AnimNextUncookedOnly/Private/Graph/AnimNextGraph_EditorData.cpp
+++ b/Engine/Plugins/Experimental/AnimNext/Source/AnimNextUncookedOnly/Private/Graph/AnimNextGraph_EditorData.cpp
@@ -14,6 +14,7 @@
 #include "RigVMModel/RigVMNotifications.h"
 #include "RigVMModel/Nodes/RigVMCollapseNode.h"
 #include "Curves/CurveFloat.h"
+#include "UObject/AssetRegistryTagsContext.h"
 #include "UObject/ObjectSaveContext.h"
 #include "Kismet2/BlueprintEditorUtils.h"
 #include "Graph/AnimNextGraphEntry.h"
@@ -175,7 +176,14 @@ void UAnimNextGraph_EditorData::PostDuplicate(EDuplicateMode::Type DuplicateMode
 #if WITH_EDITOR
 void UAnimNextGraph_EditorData::GetAssetRegistryTags(TArray<FAssetRegistryTag>& OutTags) const
 {
-	UObject::GetAssetRegistryTags(OutTags);
+	PRAGMA_DISABLE_DEPRECATION_WARNINGS;
+	Super::GetAssetRegistryTags(OutTags);
+	PRAGMA_ENABLE_DEPRECATION_WARNINGS;
+}
+
+void UAnimNextGraph_EditorData::GetAssetRegistryTags(FAssetRegistryTagsContext Context) const
+{
+	Super::GetAssetRegistryTags(Context);
 
 	FAnimNextParameterProviderAssetRegistryExports ExportParameters;
 	
@@ -191,7 +199,7 @@ void UAnimNextGraph_EditorData::GetAssetRegistryTags(TArray<FAssetRegistryTag>&
 	{
 		FString TagValue;	
 		FAnimNextParameterProviderAssetRegistryExports::StaticStruct()->ExportText(TagValue, &ExportParameters, nullptr, nullptr, PPF_None, nullptr);
-		OutTags.Add(FAssetRegistryTag(UE::AnimNext::ExportsAnimNextAssetRegistryTag, TagValue, FAssetRegistryTag::TT_Hidden));	
+		Context.AddTag(FAssetRegistryTag(UE::AnimNext::ExportsAnimNextAssetRegistryTag, TagValue, FAssetRegistryTag::TT_Hidden));	
 	}
 }
 
diff --git a/Engine/Plugins/Experimental/AnimNext/Source/AnimNextUncookedOnly/Private/Param/AnimNextParameterBlock_EditorData.cpp b/Engine/Plugins/Experimental/AnimNext/Source/AnimNextUncookedOnly/Private/Param/AnimNextParameterBlock_EditorData.cpp
--- a/Engine/Plugins/Experimental/AnimNext/Source/AnimNextUncookedOnly/Private/Param/AnimNextParameterBlock_EditorData.cpp
+++ b/Engine/Plugins/Experimental/AnimNext/Source/AnimNextUncookedOnly/Private/Param/AnimNextParameterBlock_EditorData.cpp
@@ -23,6 +23,7 @@
 #include "Param/AnimNextParameterBlockParameter.h"
 #include "Param/RigVMDispatch_GetLayerParameter.h"
 #include "Param/RigVMDispatch_GetParameter.h"
+#include "UObject/AssetRegistryTagsContext.h"
 
 #if WITH_EDITORONLY_DATA
 
@@ -580,7 +581,14 @@ void UAnimNextParameterBlock_EditorData::PostTransacted(const FTransactionObject
 
 void UAnimNextParameterBlock_EditorData::GetAssetRegistryTags(TArray<FAssetRegistryTag>& OutTags) const
 {
+	PRAGMA_DISABLE_DEPRECATION_WARNINGS;
 	Super::GetAssetRegistryTags(OutTags);
+	PRAGMA_ENABLE_DEPRECATION_WARNINGS;
+}
+
+void UAnimNextParameterBlock_EditorData::GetAssetRegistryTags(FAssetRegistryTagsContext Context) const
+{
+	Super::GetAssetRegistryTags(Context);
 
 	FAnimNextParameterProviderAssetRegistryExports ExportParameters;
 	ExportParameters.Parameters.Reserve(Entries.Num());
@@ -603,7 +611,7 @@ void UAnimNextParameterBlock_EditorData::GetAssetRegistryTags(TArray<FAssetRegis
 
 	FString TagValue;
 	FAnimNextParameterProviderAssetRegistryExports::StaticStruct()->ExportText(TagValue, &ExportParameters, nullptr, nullptr, PPF_None, nullptr);
-	OutTags.Add(FAssetRegistryTag(UE::AnimNext::ExportsAnimNextAssetRegistryTag, TagValue, FAssetRegistryTag::TT_Hidden));
+	Context.AddTag(FAssetRegistryTag(UE::AnimNext::ExportsAnimNextAssetRegistryTag, TagValue, FAssetRegistryTag::TT_Hidden));
 }
 
 #if WITH_EDITOR
diff --git a/Engine/Plugins/Experimental/Animation/ContextualAnimation/Source/ContextualAnimationEditor/Private/SContextualAnimAssetBrowser.cpp b/Engine/Plugins/Experimental/Animation/ContextualAnimation/Source/ContextualAnimationEditor/Private/SContextualAnimAssetBrowser.cpp
--- a/Engine/Plugins/Experimental/Animation/ContextualAnimation/Source/ContextualAnimationEditor/Private/SContextualAnimAssetBrowser.cpp
+++ b/Engine/Plugins/Experimental/Animation/ContextualAnimation/Source/ContextualAnimationEditor/Private/SContextualAnimAssetBrowser.cpp
@@ -6,6 +6,7 @@
 #include "Animation/AnimSequence.h"
 #include "Animation/AnimMontage.h"
 #include "Widgets/SBoxPanel.h"
+#include "UObject/AssetRegistryTagsContext.h"
 
 void SContextualAnimAssetBrowser::Construct(const FArguments& InArgs)
 {
@@ -29,11 +30,12 @@ void SContextualAnimAssetBrowser::Construct(const FArguments& InArgs)
 	AssetPickerConfig.bShowTypeInColumnView = false;
 
 	// hide all asset registry columns by default (we only really want the name and path)
-	TArray<UObject::FAssetRegistryTag> AssetRegistryTags;
-	UAnimSequence::StaticClass()->GetDefaultObject()->GetAssetRegistryTags(AssetRegistryTags);
-	for (UObject::FAssetRegistryTag& AssetRegistryTag : AssetRegistryTags)
+	UObject* AnimSequenceDefaultObject = UAnimSequence::StaticClass()->GetDefaultObject();
+	FAssetRegistryTagsContextData TagsContext(AnimSequenceDefaultObject, EAssetRegistryTagsCaller::Uncategorized);
+	AnimSequenceDefaultObject->GetAssetRegistryTags(TagsContext);
+	for (const TPair<FName, UObject::FAssetRegistryTag>& TagPair : TagsContext.Tags)
 	{
-		AssetPickerConfig.HiddenColumnNames.Add(AssetRegistryTag.Name.ToString());
+		AssetPickerConfig.HiddenColumnNames.Add(TagPair.Key.ToString());
 	}
 
 	// Also hide the type column by default (but allow users to enable it, so don't use bShowTypeInColumnView)
diff --git a/Engine/Plugins/Experimental/ChaosCaching/Source/ChaosCaching/Private/Chaos/CacheCollection.cpp b/Engine/Plugins/Experimental/ChaosCaching/Source/ChaosCaching/Private/Chaos/CacheCollection.cpp
--- a/Engine/Plugins/Experimental/ChaosCaching/Source/ChaosCaching/Private/Chaos/CacheCollection.cpp
+++ b/Engine/Plugins/Experimental/ChaosCaching/Source/ChaosCaching/Private/Chaos/CacheCollection.cpp
@@ -4,15 +4,23 @@
 #include "Chaos/ChaosCache.h"
 #include "Algo/Find.h"
 #include "Async/ParallelFor.h"
+#include "UObject/AssetRegistryTagsContext.h"
 
 #include UE_INLINE_GENERATED_CPP_BY_NAME(CacheCollection)
 
 void UChaosCacheCollection::GetAssetRegistryTags(TArray<FAssetRegistryTag>& OutTags) const
+{
+	PRAGMA_DISABLE_DEPRECATION_WARNINGS;
+	Super::GetAssetRegistryTags(OutTags);
+	PRAGMA_ENABLE_DEPRECATION_WARNINGS;
+}
+
+void UChaosCacheCollection::GetAssetRegistryTags(FAssetRegistryTagsContext Context) const
 {
 #if WITH_EDITOR
 	int32 NumCaches = Caches.Num();
 
-	OutTags.Add(FAssetRegistryTag(TEXT("Number of Observed Components"), FString::FromInt(NumCaches), FAssetRegistryTag::TT_Numerical));
+	Context.AddTag(FAssetRegistryTag(TEXT("Number of Observed Components"), FString::FromInt(NumCaches), FAssetRegistryTag::TT_Numerical));
 
 	float MaxDuration = 0.0;
 	uint32 MaxRecordedFrames = 0;
@@ -34,14 +42,14 @@ void UChaosCacheCollection::GetAssetRegistryTags(TArray<FAssetRegistryTag>& OutT
 	const int64 TransformStorageSizeInBytes = sizeof(FVector3f) + sizeof(FQuat4f) + sizeof(FVector3f) + sizeof(float); // pos + rot + scale + timestamp
 	const double EstimatedTransformMemoryInMegaBytes = double(TotalTransformKeys * TransformStorageSizeInBytes) / 1024.0 / 1024.0;
 
-	OutTags.Add(FAssetRegistryTag(TEXT("Recorded Duration"), FString::Printf(TEXT("%.2f"), MaxDuration), FAssetRegistryTag::TT_Numerical));
-	OutTags.Add(FAssetRegistryTag(TEXT("Recorded Frames"), FString::Printf(TEXT("%d"), MaxRecordedFrames), FAssetRegistryTag::TT_Numerical));
-	OutTags.Add(FAssetRegistryTag(TEXT("Total Tracks"), FString::Printf(TEXT("%lld"), TotalTracks), FAssetRegistryTag::TT_Numerical));
-	OutTags.Add(FAssetRegistryTag(TEXT("Total Transform Keys"), FString::Printf(TEXT("%lld"), TotalTransformKeys), FAssetRegistryTag::TT_Numerical));
-	OutTags.Add(FAssetRegistryTag(TEXT("Estimated Transform Memory"), FString::Printf(TEXT("%.2f MB"), EstimatedTransformMemoryInMegaBytes), FAssetRegistryTag::TT_Numerical));
+	Context.AddTag(FAssetRegistryTag(TEXT("Recorded Duration"), FString::Printf(TEXT("%.2f"), MaxDuration), FAssetRegistryTag::TT_Numerical));
+	Context.AddTag(FAssetRegistryTag(TEXT("Recorded Frames"), FString::Printf(TEXT("%d"), MaxRecordedFrames), FAssetRegistryTag::TT_Numerical));
+	Context.AddTag(FAssetRegistryTag(TEXT("Total Tracks"), FString::Printf(TEXT("%lld"), TotalTracks), FAssetRegistryTag::TT_Numerical));
+	Context.AddTag(FAssetRegistryTag(TEXT("Total Transform Keys"), FString::Printf(TEXT("%lld"), TotalTransformKeys), FAssetRegistryTag::TT_Numerical));
+	Context.AddTag(FAssetRegistryTag(TEXT("Estimated Transform Memory"), FString::Printf(TEXT("%.2f MB"), EstimatedTransformMemoryInMegaBytes), FAssetRegistryTag::TT_Numerical));
 #endif
 
-	Super::GetAssetRegistryTags(OutTags);
+	Super::GetAssetRegistryTags(Context);
 }
 
 UChaosCache* UChaosCacheCollection::FindCache(const FName& CacheName) const
diff --git a/Engine/Plugins/Experimental/ChaosCaching/Source/ChaosCaching/Public/Chaos/CacheCollection.h b/Engine/Plugins/Experimental/ChaosCaching/Source/ChaosCaching/Public/Chaos/CacheCollection.h
--- a/Engine/Plugins/Experimental/ChaosCaching/Source/ChaosCaching/Public/Chaos/CacheCollection.h
+++ b/Engine/Plugins/Experimental/ChaosCaching/Source/ChaosCaching/Public/Chaos/CacheCollection.h
@@ -13,6 +13,8 @@ class UChaosCacheCollection : public UObject
 public:
 
 	//~ Begin UObject Interface
+	CHAOSCACHING_API virtual void GetAssetRegistryTags(FAssetRegistryTagsContext Context) const override;
+	UE_DEPRECATED(5.4, "Implement the version that takes FAssetRegistryTagsContext instead.")
 	CHAOSCACHING_API virtual void GetAssetRegistryTags(TArray<FAssetRegistryTag>& OutTags) const override;
 	//~ End UObject Interface
 	
diff --git a/Engine/Plugins/Experimental/Chooser/Source/Chooser/Internal/Chooser.h b/Engine/Plugins/Experimental/Chooser/Source/Chooser/Internal/Chooser.h
--- a/Engine/Plugins/Experimental/Chooser/Source/Chooser/Internal/Chooser.h
+++ b/Engine/Plugins/Experimental/Chooser/Source/Chooser/Internal/Chooser.h
@@ -25,8 +25,10 @@ public:
 	virtual void PostLoad() override;
 	virtual void Compile(bool bForce = false) override;
 #if WITH_EDITOR
+	virtual void GetAssetRegistryTags(FAssetRegistryTagsContext Context) const override;
+	UE_DEPRECATED(5.4, "Implement the version that takes FAssetRegistryTagsContext instead.")
 	virtual void GetAssetRegistryTags(TArray<FAssetRegistryTag>& OutTags) const override;
-	
+
 	void OnDependentStructChanged(UUserDefinedStruct* Blueprint) { Compile(true); }
 	void OnDependencyCompiled(UBlueprint* Blueprint) { Compile(true); }
 	virtual void AddCompileDependency(const UStruct* Struct) override;
diff --git a/Engine/Plugins/Experimental/Chooser/Source/Chooser/Private/Chooser.cpp b/Engine/Plugins/Experimental/Chooser/Source/Chooser/Private/Chooser.cpp
--- a/Engine/Plugins/Experimental/Chooser/Source/Chooser/Private/Chooser.cpp
+++ b/Engine/Plugins/Experimental/Chooser/Source/Chooser/Private/Chooser.cpp
@@ -8,6 +8,7 @@
 #include "Engine/Blueprint.h"
 #include "ChooserIndexArray.h"
 #include "IChooserParameterGameplayTag.h"
+#include "UObject/AssetRegistryTagsContext.h"
 
 #include UE_INLINE_GENERATED_CPP_BY_NAME(Chooser)
 
@@ -92,7 +93,14 @@ void UChooserTable::BeginDestroy()
 
 void UChooserTable::GetAssetRegistryTags(TArray<FAssetRegistryTag>& OutTags) const
 {
+	PRAGMA_DISABLE_DEPRECATION_WARNINGS;
 	Super::GetAssetRegistryTags(OutTags);
+	PRAGMA_ENABLE_DEPRECATION_WARNINGS;
+}
+
+void UChooserTable::GetAssetRegistryTags(FAssetRegistryTagsContext Context) const
+{
+	Super::GetAssetRegistryTags(Context);
 
 	// Output property names we use
 	TStringBuilder<256> PropertyNamesBuilder;
@@ -109,7 +117,7 @@ void UChooserTable::GetAssetRegistryTags(TArray<FAssetRegistryTag>& OutTags) con
 		}
 	}
 
-	OutTags.Add(FAssetRegistryTag(PropertyNamesTag, PropertyNamesBuilder.ToString(), FAssetRegistryTag::TT_Hidden));
+	Context.AddTag(FAssetRegistryTag(PropertyNamesTag, PropertyNamesBuilder.ToString(), FAssetRegistryTag::TT_Hidden));
 }
 
 void UChooserTable::AddCompileDependency(const UStruct* InStructType)
diff --git a/Engine/Plugins/Experimental/Chooser/Source/ProxyTable/Private/ProxyAsset.cpp b/Engine/Plugins/Experimental/Chooser/Source/ProxyTable/Private/ProxyAsset.cpp
--- a/Engine/Plugins/Experimental/Chooser/Source/ProxyTable/Private/ProxyAsset.cpp
+++ b/Engine/Plugins/Experimental/Chooser/Source/ProxyTable/Private/ProxyAsset.cpp
@@ -3,22 +3,22 @@
 #include "ProxyTableFunctionLibrary.h"
 #include "AssetRegistry/AssetRegistryModule.h"
 #include "ChooserPropertyAccess.h"
+#include "UObject/AssetRegistryTagsContext.h"
 #include "UObject/Package.h"
 #include "LookupProxy.h"
 
 FName UProxyAsset::TypeTagName = "ProxyType";
 
-void UProxyAsset::GetAssetRegistryTags(TArray<FAssetRegistryTag>& OutTags) const
+void UProxyAsset::GetAssetRegistryTags(FAssetRegistryTagsContext Context) const
 {
 	FString ResultTypeName;
 	if (Type)
 	{
 		ResultTypeName = Type.GetName();
 	}
-	OutTags.Add({TypeTagName, ResultTypeName, FAssetRegistryTag::TT_Alphabetical});
+	Context.AddTag(FAssetRegistryTag(TypeTagName, ResultTypeName, FAssetRegistryTag::TT_Alphabetical));
 	
-	
-	UObject::GetAssetRegistryTags(OutTags);
+	Super::GetAssetRegistryTags(Context);
 }
 
 #if WITH_EDITOR
diff --git a/Engine/Plugins/Experimental/Chooser/Source/ProxyTable/Public/ProxyAsset.h b/Engine/Plugins/Experimental/Chooser/Source/ProxyTable/Public/ProxyAsset.h
--- a/Engine/Plugins/Experimental/Chooser/Source/ProxyTable/Public/ProxyAsset.h
+++ b/Engine/Plugins/Experimental/Chooser/Source/ProxyTable/Public/ProxyAsset.h
@@ -19,7 +19,7 @@ class UProxyAsset : public UObject, public IHasContextClass
 public:
 	UProxyAsset() {}
 
-	virtual void GetAssetRegistryTags(TArray<FAssetRegistryTag>& OutTags) const override;
+	virtual void GetAssetRegistryTags(FAssetRegistryTagsContext Context) const override;
 
 	static PROXYTABLE_API FName TypeTagName;
 
diff --git a/Engine/Plugins/Experimental/CommonConversation/Source/CommonConversationRuntime/Private/ConversationDatabase.cpp b/Engine/Plugins/Experimental/CommonConversation/Source/CommonConversationRuntime/Private/ConversationDatabase.cpp
--- a/Engine/Plugins/Experimental/CommonConversation/Source/CommonConversationRuntime/Private/ConversationDatabase.cpp
+++ b/Engine/Plugins/Experimental/CommonConversation/Source/CommonConversationRuntime/Private/ConversationDatabase.cpp
@@ -11,6 +11,7 @@
 #endif
 #include "AssetRegistry/AssetRegistryModule.h"
 #include "AssetRegistry/AssetRegistryState.h"
+#include "UObject/AssetRegistryTagsContext.h"
 
 #include UE_INLINE_GENERATED_CPP_BY_NAME(ConversationDatabase)
 
@@ -23,7 +24,14 @@ UConversationDatabase::UConversationDatabase(const FObjectInitializer& ObjectIni
 
 void UConversationDatabase::GetAssetRegistryTags(TArray<FAssetRegistryTag>& OutTags) const
 {
+	PRAGMA_DISABLE_DEPRECATION_WARNINGS;
 	Super::GetAssetRegistryTags(OutTags);
+	PRAGMA_ENABLE_DEPRECATION_WARNINGS;
+}
+
+void UConversationDatabase::GetAssetRegistryTags(FAssetRegistryTagsContext Context) const
+{
+	Super::GetAssetRegistryTags(Context);
 }
 
 #if WITH_EDITOR
diff --git a/Engine/Plugins/Experimental/GameFeatures/Source/GameFeatures/Private/GameFeatureData.cpp b/Engine/Plugins/Experimental/GameFeatures/Source/GameFeatures/Private/GameFeatureData.cpp
--- a/Engine/Plugins/Experimental/GameFeatures/Source/GameFeatures/Private/GameFeatureData.cpp
+++ b/Engine/Plugins/Experimental/GameFeatures/Source/GameFeatures/Private/GameFeatureData.cpp
@@ -7,6 +7,7 @@
 #include "InstallBundleUtils.h"
 #include "Misc/ConfigCacheIni.h"
 #include "Misc/ConfigContext.h"
+#include "UObject/AssetRegistryTagsContext.h"
 #include "UObject/CoreRedirects.h"
 #include "DeviceProfiles/DeviceProfile.h"
 #include "DeviceProfiles/DeviceProfileFragment.h"
@@ -582,7 +583,14 @@ void UGameFeatureData::GetContentBundleGuidsFromAsset(const FAssetData& Asset, T
 
 void UGameFeatureData::GetAssetRegistryTags(TArray<FAssetRegistryTag>& OutTags) const
 {
+	PRAGMA_DISABLE_DEPRECATION_WARNINGS;
 	Super::GetAssetRegistryTags(OutTags);
+	PRAGMA_ENABLE_DEPRECATION_WARNINGS;
+}
+
+void UGameFeatureData::GetAssetRegistryTags(FAssetRegistryTagsContext Context) const
+{
+	Super::GetAssetRegistryTags(Context);
 
 	TArray<FGuid> ContentBundleGuids;
 
@@ -600,7 +608,7 @@ void UGameFeatureData::GetAssetRegistryTags(TArray<FAssetRegistryTag>& OutTags)
 	if (ContentBundleGuids.Num() > 0)
 	{
 		FString ContentBundleGuidsStr = FString::JoinBy(ContentBundleGuids, TEXT(","), [&](const FGuid& Guid) { return Guid.ToString(); });
-		OutTags.Add(FAssetRegistryTag(GetContentBundleGuidsAssetRegistryTag(), ContentBundleGuidsStr, FAssetRegistryTag::TT_Hidden));
+		Context.AddTag(FAssetRegistryTag(GetContentBundleGuidsAssetRegistryTag(), ContentBundleGuidsStr, FAssetRegistryTag::TT_Hidden));
 	}
 }
 #endif
diff --git a/Engine/Plugins/Experimental/Mutable/Source/CustomizableObject/Classes/MuCO/CustomizableObject.h b/Engine/Plugins/Experimental/Mutable/Source/CustomizableObject/Classes/MuCO/CustomizableObject.h
--- a/Engine/Plugins/Experimental/Mutable/Source/CustomizableObject/Classes/MuCO/CustomizableObject.h
+++ b/Engine/Plugins/Experimental/Mutable/Source/CustomizableObject/Classes/MuCO/CustomizableObject.h
@@ -1435,7 +1435,9 @@ public:
 	FString GetCompiledDataFileName(bool bIsModel, const ITargetPlatform* InTargetPlatform = nullptr, bool bIsDiskStreamer = false);
 
 	/** Used to set the flag IsRoot */
-	void GetAssetRegistryTags(TArray<FAssetRegistryTag>& OutTags) const override;
+	virtual void GetAssetRegistryTags(FAssetRegistryTagsContext Context) const override;
+	UE_DEPRECATED(5.4, "Implement the version that takes FAssetRegistryTagsContext instead.")
+	virtual void GetAssetRegistryTags(TArray<FAssetRegistryTag>& OutTags) const override;
 
 	/** Returns a one line description of an object for viewing in the thumbnail view of the generic browser */
 	FString GetDesc() override;
diff --git a/Engine/Plugins/Experimental/Mutable/Source/CustomizableObject/Private/MuCO/CustomizableObject.cpp b/Engine/Plugins/Experimental/Mutable/Source/CustomizableObject/Private/MuCO/CustomizableObject.cpp
--- a/Engine/Plugins/Experimental/Mutable/Source/CustomizableObject/Private/MuCO/CustomizableObject.cpp
+++ b/Engine/Plugins/Experimental/Mutable/Source/CustomizableObject/Private/MuCO/CustomizableObject.cpp
@@ -29,6 +29,7 @@
 #include "PhysicsEngine/PhysicsAsset.h"
 #include "Serialization/MemoryReader.h"
 #include "Serialization/MemoryWriter.h"
+#include "UObject/AssetRegistryTagsContext.h"
 #include "UObject/ObjectSaveContext.h"
 
 #if WITH_EDITOR
@@ -74,6 +75,13 @@ void UCustomizableObject::UpdateVersionId()
 
 
 void UCustomizableObject::GetAssetRegistryTags(TArray<FAssetRegistryTag>& OutTags) const
+{
+	PRAGMA_DISABLE_DEPRECATION_WARNINGS;
+	Super::GetAssetRegistryTags(OutTags);
+	PRAGMA_ENABLE_DEPRECATION_WARNINGS;
+}
+
+void UCustomizableObject::GetAssetRegistryTags(FAssetRegistryTagsContext Context) const
 {
 	int32 isRoot = 0;
 #if WITH_EDITOR
@@ -85,8 +93,8 @@ void UCustomizableObject::GetAssetRegistryTags(TArray<FAssetRegistryTag>& OutTag
 	}
 #endif
 
-	OutTags.Add(FAssetRegistryTag("IsRoot", FString::FromInt(isRoot), FAssetRegistryTag::TT_Numerical));
-	Super::GetAssetRegistryTags(OutTags);
+	Context.AddTag(FAssetRegistryTag("IsRoot", FString::FromInt(isRoot), FAssetRegistryTag::TT_Numerical));
+	Super::GetAssetRegistryTags(Context);
 }
 
 
@@ -117,7 +125,7 @@ void UCustomizableObject::PreSave(FObjectPreSaveContext ObjectSaveContext)
 		const ITargetPlatform* TargetPlatform = ObjectSaveContext.GetTargetPlatform();
 		if (const FMutableCachedPlatformData* PlatformData = CachedPlatformsData.Find(TargetPlatform->PlatformName()))
 		{
-			// Load cached data before saving
+		// Load cached data before saving
 			FMemoryReaderView MemoryReader(PlatformData->ModelData);
 			LoadCompiledData(MemoryReader, TargetPlatform, true);
 
