[MLDeformer]
- Fixed a crash when deleting bones from the bone list in local mode.
- Make sure the ML Deformer component gets reinitialized after training and after updating its trained input info object.
- Handle the case when the deformer is not compatible with the current skeletal mesh.

#rn

diff --git a/Engine/Plugins/Animation/MLDeformer/MLDeformerFramework/Source/MLDeformerFramework/Private/MLDeformerModelInstance.cpp b/Engine/Plugins/Animation/MLDeformer/MLDeformerFramework/Source/MLDeformerFramework/Private/MLDeformerModelInstance.cpp
--- a/Engine/Plugins/Animation/MLDeformer/MLDeformerFramework/Source/MLDeformerFramework/Private/MLDeformerModelInstance.cpp
+++ b/Engine/Plugins/Animation/MLDeformer/MLDeformerFramework/Source/MLDeformerFramework/Private/MLDeformerModelInstance.cpp
@@ -402,7 +402,7 @@ void UMLDeformerModelInstance::Tick(float DeltaTime, float ModelWeight)
 	}
 
 	bool bExecuteCalled = false;
-	if (ModelWeight > 0.0001f && HasValidTransforms() && SetupInputs())
+	if (ModelWeight > 0.0001f && IsCompatible() && HasValidTransforms() && SetupInputs())
 	{
 		// Execute the model instance.
 		// For models using neural networks this will perform the inference, 
diff --git a/Engine/Plugins/Animation/MLDeformer/MLDeformerFramework/Source/MLDeformerFrameworkEditor/Private/MLDeformerEditorModel.cpp b/Engine/Plugins/Animation/MLDeformer/MLDeformerFramework/Source/MLDeformerFrameworkEditor/Private/MLDeformerEditorModel.cpp
--- a/Engine/Plugins/Animation/MLDeformer/MLDeformerFramework/Source/MLDeformerFrameworkEditor/Private/MLDeformerEditorModel.cpp
+++ b/Engine/Plugins/Animation/MLDeformer/MLDeformerFramework/Source/MLDeformerFrameworkEditor/Private/MLDeformerEditorModel.cpp
@@ -2223,6 +2223,7 @@ namespace UE::MLDeformer
 		SampleDeltas();
 		Model->InitGPUData();
 		UpdateMemoryUsage();
+		RefreshMLDeformerComponents();
 	}
 
 	FMLDeformerEditorActor* FMLDeformerEditorModel::GetTimelineEditorActor() const
diff --git a/Engine/Plugins/Animation/MLDeformer/NeuralMorphModel/Source/NeuralMorphModelEditor/Private/NeuralMorphEditorModel.cpp b/Engine/Plugins/Animation/MLDeformer/NeuralMorphModel/Source/NeuralMorphModelEditor/Private/NeuralMorphEditorModel.cpp
--- a/Engine/Plugins/Animation/MLDeformer/NeuralMorphModel/Source/NeuralMorphModelEditor/Private/NeuralMorphEditorModel.cpp
+++ b/Engine/Plugins/Animation/MLDeformer/NeuralMorphModel/Source/NeuralMorphModelEditor/Private/NeuralMorphEditorModel.cpp
@@ -677,13 +677,19 @@ namespace UE::NeuralMorphModel
 		FMLDeformerMaskInfo* MaskInfo = nullptr;
 		if (FinalMaskItemIndex < InputInfo->GetBoneNames().Num())	// If we have a bone.
 		{
-			MaskInfo = GetNeuralMorphModel()->BoneMaskInfoMap.Find(InputInfo->GetBoneNames()[FinalMaskItemIndex]);
+			if (InputInfo->GetBoneNames().IsValidIndex(FinalMaskItemIndex))
+			{
+				MaskInfo = GetNeuralMorphModel()->BoneMaskInfoMap.Find(InputInfo->GetBoneNames()[FinalMaskItemIndex]);
+			}
 		}
 		else // It's a bone group.
 		{
 			const int32 GroupIndex = FinalMaskItemIndex - (InputInfo->GetBoneNames().Num() + InputInfo->GetCurveNames().Num());
-			const FName GroupName = InputInfo->GetBoneGroups()[GroupIndex].GroupName;
-			MaskInfo = GetNeuralMorphModel()->BoneGroupMaskInfoMap.Find(GroupName);
+			if (InputInfo->GetBoneGroups().IsValidIndex(GroupIndex))
+			{
+				const FName GroupName = InputInfo->GetBoneGroups()[GroupIndex].GroupName;
+				MaskInfo = GetNeuralMorphModel()->BoneGroupMaskInfoMap.Find(GroupName);
+			}
 		}
 
 		if (!MaskBuffer.IsEmpty())
