[MLDeformer] Some tiny fixes that could cause crashes.

diff --git a/Engine/Plugins/Animation/MLDeformer/NeuralMorphModel/Source/NeuralMorphModelEditor/Private/NeuralMorphEditorModel.cpp b/Engine/Plugins/Animation/MLDeformer/NeuralMorphModel/Source/NeuralMorphModelEditor/Private/NeuralMorphEditorModel.cpp
--- a/Engine/Plugins/Animation/MLDeformer/NeuralMorphModel/Source/NeuralMorphModelEditor/Private/NeuralMorphEditorModel.cpp
+++ b/Engine/Plugins/Animation/MLDeformer/NeuralMorphModel/Source/NeuralMorphModelEditor/Private/NeuralMorphEditorModel.cpp
@@ -314,7 +314,6 @@ namespace UE::NeuralMorphModel
 		}
 
 		const int32 NumBaseMeshVerts = Model->GetNumBaseMeshVerts();
-		check(NumBaseMeshVerts > 0);
 
 		USkeletalMesh* SkeletalMesh = Model->GetSkeletalMesh();	
 		UNeuralMorphInputInfo* NeuralMorphInputInfo = Cast<UNeuralMorphInputInfo>(GetEditorInputInfo());
diff --git a/Engine/Source/Runtime/Engine/Private/SkeletalMesh.cpp b/Engine/Source/Runtime/Engine/Private/SkeletalMesh.cpp
--- a/Engine/Source/Runtime/Engine/Private/SkeletalMesh.cpp
+++ b/Engine/Source/Runtime/Engine/Private/SkeletalMesh.cpp
@@ -1036,7 +1036,7 @@ void USkeletalMesh::ReleaseResources()
 int32 USkeletalMesh::GetNumImportedVertices() const
 {
 	const FSkeletalMeshModel* SkeletalMeshModel = GetImportedModel();
-	if (SkeletalMeshModel)
+	if (SkeletalMeshModel && !SkeletalMeshModel->LODModels.IsEmpty())
 	{
 		const int32 MaxIndex = SkeletalMeshModel->LODModels[0].MaxImportVertex;
 		return (MaxIndex > 0) ? (MaxIndex + 1) : 0;
