[MLDeformer] Fixed a crash when there are no input sample animations, and the NMM is set to local mode, and you click a bone.

diff --git a/Engine/Plugins/Animation/MLDeformer/NeuralMorphModel/Source/NeuralMorphModelEditor/Private/NeuralMorphEditorModel.cpp b/Engine/Plugins/Animation/MLDeformer/NeuralMorphModel/Source/NeuralMorphModelEditor/Private/NeuralMorphEditorModel.cpp
--- a/Engine/Plugins/Animation/MLDeformer/NeuralMorphModel/Source/NeuralMorphModelEditor/Private/NeuralMorphEditorModel.cpp
+++ b/Engine/Plugins/Animation/MLDeformer/NeuralMorphModel/Source/NeuralMorphModelEditor/Private/NeuralMorphEditorModel.cpp
@@ -668,6 +668,10 @@ namespace UE::NeuralMorphModel
 	void FNeuralMorphEditorModel::DebugDrawItemMask(FPrimitiveDrawInterface* PDI, int32 MaskItemIndex, const FVector& DrawOffset)
 	{
 		FMLDeformerSampler* Sampler = GetSamplerForActiveAnim();
+		if (!Sampler)
+		{
+			return;
+		}
 
 		UMLDeformerMorphModelVizSettings* VizSettings = Cast<UMLDeformerMorphModelVizSettings>(Model->GetVizSettings());
 		const int32 NumVerts = Model->GetNumBaseMeshVerts();
