[MLDeformer] Fixed a crash and potentially incorrect behavior when switching from Local mode to Global mode in the Neural Morph Model.

diff --git a/Engine/Plugins/Animation/MLDeformer/NeuralMorphModel/Source/NeuralMorphModelEditor/Private/NeuralMorphEditorModel.cpp b/Engine/Plugins/Animation/MLDeformer/NeuralMorphModel/Source/NeuralMorphModelEditor/Private/NeuralMorphEditorModel.cpp
--- a/Engine/Plugins/Animation/MLDeformer/NeuralMorphModel/Source/NeuralMorphModelEditor/Private/NeuralMorphEditorModel.cpp
+++ b/Engine/Plugins/Animation/MLDeformer/NeuralMorphModel/Source/NeuralMorphModelEditor/Private/NeuralMorphEditorModel.cpp
@@ -66,10 +66,7 @@ namespace UE::NeuralMorphModel
 			Property->GetFName() == TEXT("CurveName"))		// The curve name inside one of the items in the CurveNames list changed.
 		{
 			UpdateIsReadyForTrainingState();
-			if (GetNeuralMorphModel()->GetModelMode() == ENeuralMorphMode::Local)
-			{
-				RebuildEditorMaskInfo();
-			}
+			RebuildEditorMaskInfo();
 		}
 		else if (Property->GetFName() == GET_MEMBER_NAME_CHECKED(UNeuralMorphModel, Mode))
 		{
@@ -78,10 +75,7 @@ namespace UE::NeuralMorphModel
 				SetResamplingInputOutputsNeeded(true);
 				UpdateIsReadyForTrainingState();
 				GetEditor()->GetModelDetailsView()->ForceRefresh();
-				if (GetNeuralMorphModel()->GetModelMode() == ENeuralMorphMode::Local)
-				{
-					RebuildEditorMaskInfo();
-				}
+				RebuildEditorMaskInfo();
 			}
 		}
 	}
@@ -288,7 +282,7 @@ namespace UE::NeuralMorphModel
 		check(MorphInputInfo);
 
 		const TArray<float>& MaskBuffer = MorphInputInfo->GetInputItemMaskBuffer();
-		if (MaskBuffer.IsEmpty() || !NeuralMorphModel->IsBoneMaskingEnabled())
+		if (MaskBuffer.IsEmpty() || !NeuralMorphModel->IsBoneMaskingEnabled() || NeuralMorphModel->GetModelMode() != ENeuralMorphMode::Local)
 		{
 			return TArrayView<const float>();
 		}
