[MLDeformer] Make curves work with AnimNext.

diff --git a/Engine/Plugins/Animation/MLDeformer/MLDeformerFramework/Source/MLDeformerFramework/Private/MLDeformerInputInfo.cpp b/Engine/Plugins/Animation/MLDeformer/MLDeformerFramework/Source/MLDeformerFramework/Private/MLDeformerInputInfo.cpp
--- a/Engine/Plugins/Animation/MLDeformer/MLDeformerFramework/Source/MLDeformerFramework/Private/MLDeformerInputInfo.cpp
+++ b/Engine/Plugins/Animation/MLDeformer/MLDeformerFramework/Source/MLDeformerFramework/Private/MLDeformerInputInfo.cpp
@@ -6,7 +6,6 @@
 #include "MLDeformerObjectVersion.h"
 #include "Animation/AnimCurveTypes.h"
 #include "Animation/Skeleton.h"
-#include "Animation/AnimInstance.h"
 #include "ReferenceSkeleton.h"
 #include "Engine/SkeletalMesh.h"
 #include "Components/SkeletalMeshComponent.h"
@@ -245,7 +244,6 @@ void UMLDeformerInputInfo::ExtractCurveValues(USkeletalMeshComponent* SkelMeshCo
 
 void UMLDeformerInputInfo::ExtractCurveValues(USkeletalMeshComponent* SkelMeshComponent, TArray<float>& OutValues, int32 NumFloatsPerCurve) const
 {
-	UAnimInstance* AnimInstance = SkelMeshComponent->GetAnimInstance();
 	const int32 NumCurves = CurveNames.Num();
 	const int32 NumCurveFloats = NumCurves * NumFloatsPerCurve;
 	OutValues.Reset(NumCurveFloats);
@@ -253,7 +251,9 @@ void UMLDeformerInputInfo::ExtractCurveValues(USkeletalMeshComponent* SkelMeshCo
 	for (int32 Index = 0; Index < NumCurves; ++Index)
 	{
 		const FName CurveName = CurveNames[Index];
-		OutValues[Index * NumFloatsPerCurve] = AnimInstance->GetCurveValue(CurveName);
+		float CurveValue;
+		SkelMeshComponent->GetCurveValue(CurveName, 0.0f, CurveValue);
+		OutValues[Index * NumFloatsPerCurve] = CurveValue;
 	}
 }
 
diff --git a/Engine/Plugins/Animation/MLDeformer/MLDeformerFramework/Source/MLDeformerFramework/Private/MLDeformerModelInstance.cpp b/Engine/Plugins/Animation/MLDeformer/MLDeformerFramework/Source/MLDeformerFramework/Private/MLDeformerModelInstance.cpp
--- a/Engine/Plugins/Animation/MLDeformer/MLDeformerFramework/Source/MLDeformerFramework/Private/MLDeformerModelInstance.cpp
+++ b/Engine/Plugins/Animation/MLDeformer/MLDeformerFramework/Source/MLDeformerFramework/Private/MLDeformerModelInstance.cpp
@@ -8,7 +8,6 @@
 #include "MLDeformerComponent.h"
 #include "Components/SkeletalMeshComponent.h"
 #include "Engine/SkeletalMesh.h"
-#include "Animation/AnimInstance.h"
 #include "RenderingThread.h"
 #include "RenderGraphBuilder.h"
 #include "DrawDebugHelpers.h"
@@ -439,22 +438,12 @@ int64 UMLDeformerModelInstance::SetCurveValues(float* OutputBuffer, int64 Output
 	checkf((Index + AssetNumCurves) <= OutputBufferSize, TEXT("Writing curves past the end of the input buffer"));
 
 	// Write the curve weights to the output buffer.
-	UAnimInstance* AnimInstance = SkeletalMeshComponent->GetAnimInstance();
-	if (AnimInstance)
+	for (int32 CurveIndex = 0; CurveIndex < AssetNumCurves; ++CurveIndex)
 	{
-		for (int32 CurveIndex = 0; CurveIndex < AssetNumCurves; ++CurveIndex)
-		{
-			const FName CurveName = InputInfo->GetCurveName(CurveIndex);
-			const float CurveValue = AnimInstance->GetCurveValue(CurveName);	// Outputs 0.0 when not found.
-			OutputBuffer[Index++] = CurveValue;
-		}
-	}
-	else
-	{
-		for (int32 CurveIndex = 0; CurveIndex < AssetNumCurves; ++CurveIndex)
-		{
-			OutputBuffer[Index++] = 0.0f;
-		}
+		const FName CurveName = InputInfo->GetCurveName(CurveIndex);
+		float CurveValue;
+		SkeletalMeshComponent->GetCurveValue(CurveName, 0.0f, CurveValue);
+		OutputBuffer[Index++] = CurveValue;
 	}
 
 	return Index;
diff --git a/Engine/Plugins/Animation/MLDeformer/NeuralMorphModel/Source/NeuralMorphModel/Private/NeuralMorphModelInstance.cpp b/Engine/Plugins/Animation/MLDeformer/NeuralMorphModel/Source/NeuralMorphModel/Private/NeuralMorphModelInstance.cpp
--- a/Engine/Plugins/Animation/MLDeformer/NeuralMorphModel/Source/NeuralMorphModel/Private/NeuralMorphModelInstance.cpp
+++ b/Engine/Plugins/Animation/MLDeformer/NeuralMorphModel/Source/NeuralMorphModel/Private/NeuralMorphModelInstance.cpp
@@ -6,7 +6,6 @@
 #include "NeuralMorphInputInfo.h"
 #include "MLDeformerComponent.h"
 #include "MLDeformerAsset.h"
-#include "Animation/AnimInstance.h"
 #include "Components/ExternalMorphSet.h"
 #include "Components/SkeletalMeshComponent.h"
 
@@ -140,37 +139,23 @@ int64 UNeuralMorphModelInstance::SetCurveValues(float* OutputBuffer, int64 Outpu
 		}
 
 		// Write the actual curve values.
-		UAnimInstance* AnimInstance = SkeletalMeshComponent->GetAnimInstance();
-		if (AnimInstance)
+		for (int32 CurveIndex = 0; CurveIndex < AssetNumCurves; ++CurveIndex)
 		{
-			for (int32 CurveIndex = 0; CurveIndex < AssetNumCurves; ++CurveIndex)
-			{
-				const FName CurveName = InputInfo->GetCurveName(CurveIndex);
-				const float CurveValue = AnimInstance->GetCurveValue(CurveName);	// Outputs 0.0 when not found.
-				OutputBuffer[StartIndex + CurveIndex * NumFloatsPerCurve] = CurveValue;
-			}
+			const FName CurveName = InputInfo->GetCurveName(CurveIndex);
+			float CurveValue;
+			SkeletalMeshComponent->GetCurveValue(CurveName, 0.0f, CurveValue);
+			OutputBuffer[StartIndex + CurveIndex * NumFloatsPerCurve] = CurveValue;
 		}
 	}
 	else
 	{
 		checkSlow(NumFloatsPerCurve == 1);
-		UAnimInstance* AnimInstance = SkeletalMeshComponent->GetAnimInstance();
-		if (AnimInstance)
-		{
-			for (int32 CurveIndex = 0; CurveIndex < AssetNumCurves; ++CurveIndex)
-			{
-				const FName CurveName = InputInfo->GetCurveName(CurveIndex);
-				const float CurveValue = AnimInstance->GetCurveValue(CurveName);	// Outputs 0.0 when not found.
-				OutputBuffer[Index++] = CurveValue;
-			}
-		}
-		else
+		for (int32 CurveIndex = 0; CurveIndex < AssetNumCurves; ++CurveIndex)
 		{
-			// Just write zeros.
-			for (int32 CurveIndex = 0; CurveIndex < NumCurveFloats; ++CurveIndex)
-			{
-				OutputBuffer[Index++] = 0.0f;
-			}
+			const FName CurveName = InputInfo->GetCurveName(CurveIndex);
+			float CurveValue;
+			SkeletalMeshComponent->GetCurveValue(CurveName, 0.0f, CurveValue);
+			OutputBuffer[Index++] = CurveValue;
 		}
 	}
 
diff --git a/Engine/Plugins/Experimental/AnimNext/Source/AnimNext/Private/Component/AnimNextComponent.cpp b/Engine/Plugins/Experimental/AnimNext/Source/AnimNext/Private/Component/AnimNextComponent.cpp
--- a/Engine/Plugins/Experimental/AnimNext/Source/AnimNext/Private/Component/AnimNextComponent.cpp
+++ b/Engine/Plugins/Experimental/AnimNext/Source/AnimNext/Private/Component/AnimNextComponent.cpp
@@ -940,7 +940,10 @@ void UAnimNextComponent::AddPrerequisite(UObject* InObject, FTickFunction& InTic
 
 void UAnimNextComponent::AddComponentPrerequisite(UActorComponent* InComponent, FName InEventName)
 {
-	AddPrerequisite(InComponent, InComponent->PrimaryComponentTick, InEventName);
+	if (InComponent)
+	{
+		AddPrerequisite(InComponent, InComponent->PrimaryComponentTick, InEventName);
+	}
 }
 
 void UAnimNextComponent::AddSubsequent(UObject* InObject, FTickFunction& InTickFunction, FName InEventName)
@@ -953,7 +956,10 @@ void UAnimNextComponent::AddSubsequent(UObject* InObject, FTickFunction& InTickF
 
 void UAnimNextComponent::AddComponentSubsequent(UActorComponent* InComponent, FName InEventName)
 {
-	AddSubsequent(InComponent, InComponent->PrimaryComponentTick, InEventName);
+	if (InComponent)
+	{
+		AddSubsequent(InComponent, InComponent->PrimaryComponentTick, InEventName);
+	}
 }
 
 void UAnimNextComponent::RemovePrerequisite(UObject* InObject, FTickFunction& InTickFunction, FName InEventName)
@@ -966,7 +972,10 @@ void UAnimNextComponent::RemovePrerequisite(UObject* InObject, FTickFunction& In
 
 void UAnimNextComponent::RemoveComponentPrerequisite(UActorComponent* InComponent, FName InEventName)
 {
-	RemovePrerequisite(InComponent, InComponent->PrimaryComponentTick, InEventName);
+	if (InComponent)
+	{
+		RemovePrerequisite(InComponent, InComponent->PrimaryComponentTick, InEventName);
+	}
 }
 
 void UAnimNextComponent::RemoveSubsequent(UObject* InObject, FTickFunction& InTickFunction, FName InEventName)
@@ -979,7 +988,10 @@ void UAnimNextComponent::RemoveSubsequent(UObject* InObject, FTickFunction& InTi
 
 void UAnimNextComponent::RemoveComponentSubsequent(UActorComponent* InComponent, FName InEventName)
 {
-	RemoveSubsequent(InComponent, InComponent->PrimaryComponentTick, InEventName);
+	if (InComponent)
+	{
+		RemoveSubsequent(InComponent, InComponent->PrimaryComponentTick, InEventName);
+	}
 }
 
 void UAnimNextComponent::AddModuleEventPrerequisite(FName InEventName, UAnimNextComponent* OtherAnimNextComponent, FName OtherEventName)
