Updated headers using UnrealCodeFixup to make sure dllstorage is on methods/staticvars instead of types. Part 1/n

diff --git a/Engine/Plugins/Animation/MLDeformer/MLDeformerFramework/Source/MLDeformerFrameworkEditor/Public/Tools/TrainingDataProcessor/TrainingDataProcessor.h b/Engine/Plugins/Animation/MLDeformer/MLDeformerFramework/Source/MLDeformerFrameworkEditor/Public/Tools/TrainingDataProcessor/TrainingDataProcessor.h
--- a/Engine/Plugins/Animation/MLDeformer/MLDeformerFramework/Source/MLDeformerFrameworkEditor/Public/Tools/TrainingDataProcessor/TrainingDataProcessor.h
+++ b/Engine/Plugins/Animation/MLDeformer/MLDeformerFramework/Source/MLDeformerFrameworkEditor/Public/Tools/TrainingDataProcessor/TrainingDataProcessor.h
@@ -5,6 +5,8 @@
 #include "CoreMinimal.h"
 #include "MLDeformerTrainingDataProcessorSettings.h"
 
+#define UE_API MLDEFORMERFRAMEWORKEDITOR_API
+
 struct FReferenceSkeleton;
 class USkeleton;
 
@@ -16,7 +18,7 @@ namespace UE::MLDeformer::TrainingDataProcessor
 	 * After that it can remix the poses based on a set of bone groups. Basically the keyframes of the
 	 * bones inside the groups get shuffled. So all the bones in the bone group will change together.
 	 */
-	class MLDEFORMERFRAMEWORKEDITOR_API FTrainingDataProcessor final
+	class FTrainingDataProcessor final
 	{
 	public:
 		/**
@@ -26,22 +28,22 @@ namespace UE::MLDeformer::TrainingDataProcessor
 		 * @param Skeleton The skeleton to use when sampling animations and apply the transforms on in the output animation.
 		 * @return Returns true if successful, or false if not. Errors will be logged in case false is returned. 
 		 */
-		bool Execute(const UMLDeformerTrainingDataProcessorSettings& Settings, const USkeleton* Skeleton);
+		UE_API bool Execute(const UMLDeformerTrainingDataProcessorSettings& Settings, const USkeleton* Skeleton);
 
 	private:
-		void Clear();
-		bool SampleFrames(const UMLDeformerTrainingDataProcessorSettings& Settings, const USkeleton& Skeleton);
-		TArray<int32> FindBestFrameIndices(int32 MaxNumFrames, const FReferenceSkeleton& RefSkeleton) const;
-		bool RemixPoses(const UMLDeformerTrainingDataProcessorSettings& Settings, const FReferenceSkeleton& RefSkeleton);
-		bool SaveAnimationDataInAnimSequence(const UMLDeformerTrainingDataProcessorSettings& Settings, const FReferenceSkeleton& RefSkeleton, const TArray<int32>& FramesToInclude);
-		int32 GetNumFrames() const;
-		TArrayView<FTransform3f> GetFrameTransforms(int32 FrameIndex);
-		TConstArrayView<FTransform3f> GetFrameTransforms(int32 FrameIndex) const;
-		TConstArrayView<FTransform3f> GetFrameTransforms(const TArray<FTransform3f>& Transforms, int32 FrameIndex) const;
-		TConstArrayView<FTransform3f> GetRefPoseTransforms() const;
-		double CalculateMeanError(TConstArrayView<FTransform3f> PoseA, TConstArrayView<FTransform3f> PoseB) const;
-		double CalculateMeanError(int32 PreviousPoseIndex, int32 CurrentPoseIndex) const;
-		static FTransform3f GetRefPoseTransform(const FReferenceSkeleton& RefSkeleton, FName BoneName);
+		UE_API void Clear();
+		UE_API bool SampleFrames(const UMLDeformerTrainingDataProcessorSettings& Settings, const USkeleton& Skeleton);
+		UE_API TArray<int32> FindBestFrameIndices(int32 MaxNumFrames, const FReferenceSkeleton& RefSkeleton) const;
+		UE_API bool RemixPoses(const UMLDeformerTrainingDataProcessorSettings& Settings, const FReferenceSkeleton& RefSkeleton);
+		UE_API bool SaveAnimationDataInAnimSequence(const UMLDeformerTrainingDataProcessorSettings& Settings, const FReferenceSkeleton& RefSkeleton, const TArray<int32>& FramesToInclude);
+		UE_API int32 GetNumFrames() const;
+		UE_API TArrayView<FTransform3f> GetFrameTransforms(int32 FrameIndex);
+		UE_API TConstArrayView<FTransform3f> GetFrameTransforms(int32 FrameIndex) const;
+		UE_API TConstArrayView<FTransform3f> GetFrameTransforms(const TArray<FTransform3f>& Transforms, int32 FrameIndex) const;
+		UE_API TConstArrayView<FTransform3f> GetRefPoseTransforms() const;
+		UE_API double CalculateMeanError(TConstArrayView<FTransform3f> PoseA, TConstArrayView<FTransform3f> PoseB) const;
+		UE_API double CalculateMeanError(int32 PreviousPoseIndex, int32 CurrentPoseIndex) const;
+		static UE_API FTransform3f GetRefPoseTransform(const FReferenceSkeleton& RefSkeleton, FName BoneName);
 		
 	private:
 		/**
@@ -74,3 +76,5 @@ namespace UE::MLDeformer::TrainingDataProcessor
 		FAnimFrameData AnimFrameData;
 	};
 } // namespace UE::MLDeformer::TrainingDataProcessor
+
+#undef UE_API
diff --git a/Engine/Plugins/Animation/MLDeformer/MLDeformerFramework/Source/MLDeformerFrameworkEditor/Public/Tools/TrainingDataProcessor/TrainingDataProcessorSettingsDetailCustomization.h b/Engine/Plugins/Animation/MLDeformer/MLDeformerFramework/Source/MLDeformerFrameworkEditor/Public/Tools/TrainingDataProcessor/TrainingDataProcessorSettingsDetailCustomization.h
--- a/Engine/Plugins/Animation/MLDeformer/MLDeformerFramework/Source/MLDeformerFrameworkEditor/Public/Tools/TrainingDataProcessor/TrainingDataProcessorSettingsDetailCustomization.h
+++ b/Engine/Plugins/Animation/MLDeformer/MLDeformerFramework/Source/MLDeformerFrameworkEditor/Public/Tools/TrainingDataProcessor/TrainingDataProcessorSettingsDetailCustomization.h
@@ -6,6 +6,8 @@
 #include "Input/Reply.h"
 #include "Layout/Visibility.h"
 
+#define UE_API MLDEFORMERFRAMEWORKEDITOR_API
+
 class UMLDeformerModel;
 class UMLDeformerTrainingDataProcessorSettings;
 class USkeleton;
@@ -16,26 +18,28 @@ namespace UE::MLDeformer::TrainingDataProcessor
 	/**
 	 * The detail customization class for the UMLDeformerTrainingDataProcessorSettings class.
 	 */
-	class MLDEFORMERFRAMEWORKEDITOR_API FTrainingDataProcessorSettingsDetailCustomization final : public IDetailCustomization
+	class FTrainingDataProcessorSettingsDetailCustomization final : public IDetailCustomization
 	{
 	public:
-		static TSharedRef<IDetailCustomization> MakeInstance();
+		static UE_API TSharedRef<IDetailCustomization> MakeInstance();
 
-		virtual void CustomizeDetails(IDetailLayoutBuilder& DetailBuilder) override;
+		UE_API virtual void CustomizeDetails(IDetailLayoutBuilder& DetailBuilder) override;
 
 	private:
-		EVisibility GetNoInputBonesWarningVisibility() const;
-		EVisibility GetNoFramesWarningVisibility() const;
-		EVisibility GetSkeletonMismatchErrorVisibility() const;
+		UE_API EVisibility GetNoInputBonesWarningVisibility() const;
+		UE_API EVisibility GetNoFramesWarningVisibility() const;
+		UE_API EVisibility GetSkeletonMismatchErrorVisibility() const;
 		
-		FReply OnCreateNewButtonClicked() const;
-		int32 GetTotalNumInputFrames() const;
-		bool FilterAnimSequences(const FAssetData& AssetData) const;
+		UE_API FReply OnCreateNewButtonClicked() const;
+		UE_API int32 GetTotalNumInputFrames() const;
+		UE_API bool FilterAnimSequences(const FAssetData& AssetData) const;
 		
-		static void Refresh(IDetailLayoutBuilder* DetailBuilder);
-		static FString FindDefaultAnimSequencePath(const UMLDeformerModel* Model);
+		static UE_API void Refresh(IDetailLayoutBuilder* DetailBuilder);
+		static UE_API FString FindDefaultAnimSequencePath(const UMLDeformerModel* Model);
 
 	private:
 		TWeakObjectPtr<UMLDeformerTrainingDataProcessorSettings> TrainingDataProcessorSettings;
 	};
 } // namespace UE::MLDeformer::TrainingDataProcessor
+
+#undef UE_API
diff --git a/Engine/Plugins/Animation/MLDeformer/NearestNeighborModel/Source/NearestNeighborModelEditor/Private/NearestNeighborEditorHelpers.h b/Engine/Plugins/Animation/MLDeformer/NearestNeighborModel/Source/NearestNeighborModelEditor/Private/NearestNeighborEditorHelpers.h
--- a/Engine/Plugins/Animation/MLDeformer/NearestNeighborModel/Source/NearestNeighborModelEditor/Private/NearestNeighborEditorHelpers.h
+++ b/Engine/Plugins/Animation/MLDeformer/NearestNeighborModel/Source/NearestNeighborModelEditor/Private/NearestNeighborEditorHelpers.h
@@ -6,23 +6,25 @@
 
 #include "NearestNeighborEditorHelpers.generated.h"
 
+#define UE_API NEARESTNEIGHBORMODELEDITOR_API
+
 class UAnimSequence;
-UCLASS(Blueprintable)
-class NEARESTNEIGHBORMODELEDITOR_API UNearestNeighborAnimStream : public UObject
+UCLASS(MinimalAPI, Blueprintable)
+class UNearestNeighborAnimStream : public UObject
 {
 	GENERATED_BODY()
 public:
 	UFUNCTION(BlueprintCallable, Category = "Python")
-	void Init(USkeleton* InSkeleton);
+	UE_API void Init(USkeleton* InSkeleton);
 
 	UFUNCTION(BlueprintPure, Category = "Python")
-	bool IsValid() const;
+	UE_API bool IsValid() const;
 
 	UFUNCTION(BlueprintCallable, Category = "Python")
-	bool AppendFrames(const UAnimSequence* Anim, TArray<int32> Frames);
+	UE_API bool AppendFrames(const UAnimSequence* Anim, TArray<int32> Frames);
 
 	UFUNCTION(BlueprintPure, Category = "Python")
-	bool ToAnim(UAnimSequence* OutAnim) const;
+	UE_API bool ToAnim(UAnimSequence* OutAnim) const;
 
 private:
 	TObjectPtr<USkeleton> Skeleton;
@@ -32,26 +34,28 @@ private:
 };
 
 class UGeometryCache;
-UCLASS(Blueprintable)
-class NEARESTNEIGHBORMODELEDITOR_API UNearestNeighborGeometryCacheStream : public UObject
+UCLASS(MinimalAPI, Blueprintable)
+class UNearestNeighborGeometryCacheStream : public UObject
 {
 	GENERATED_BODY()
 public:
 	UFUNCTION(BlueprintCallable, Category = "Python")
-	void Init(const UGeometryCache* InTemplateCache);
+	UE_API void Init(const UGeometryCache* InTemplateCache);
 
 	UFUNCTION(BlueprintPure, Category = "Python")
-	bool IsValid() const;
+	UE_API bool IsValid() const;
 
 	UFUNCTION(BlueprintCallable, Category = "Python")
-	bool AppendFrames(const UGeometryCache* Cache, TArray<int32> Frames);
+	UE_API bool AppendFrames(const UGeometryCache* Cache, TArray<int32> Frames);
 
 	UFUNCTION(BlueprintPure, Category = "Python")
-	bool ToGeometryCache(UGeometryCache* OutCache);
+	UE_API bool ToGeometryCache(UGeometryCache* OutCache);
 
 private:
 	const UGeometryCache* TemplateCache = nullptr;
 	int32 TemplateNumTracks = 0;
 	TArray<int32> TemplateTrackNumVertices;
 	TArray<TArray<TArray<FVector3f>>> TrackToFrameToPositions;
-};
\ No newline at end of file
+};
+
+#undef UE_API
diff --git a/Engine/Plugins/Animation/MLDeformer/NearestNeighborModel/Source/NearestNeighborModelEditor/Private/Tools/NearestNeighborKMeansTool.h b/Engine/Plugins/Animation/MLDeformer/NearestNeighborModel/Source/NearestNeighborModelEditor/Private/Tools/NearestNeighborKMeansTool.h
--- a/Engine/Plugins/Animation/MLDeformer/NearestNeighborModel/Source/NearestNeighborModelEditor/Private/Tools/NearestNeighborKMeansTool.h
+++ b/Engine/Plugins/Animation/MLDeformer/NearestNeighborModel/Source/NearestNeighborModelEditor/Private/Tools/NearestNeighborKMeansTool.h
@@ -21,7 +21,7 @@ namespace UE::MLDeformer
 class UNearestNeighborKMeansData;
 
 USTRUCT(BlueprintType)
-struct NEARESTNEIGHBORMODELEDITOR_API FNearestNeighborKMeansInputData
+struct FNearestNeighborKMeansInputData
 {
 	GENERATED_BODY()
 public:
@@ -36,8 +36,8 @@ public:
 };
 
 
-UCLASS(Blueprintable)
-class NEARESTNEIGHBORMODELEDITOR_API UNearestNeighborKMeansData : public UObject
+UCLASS(MinimalAPI, Blueprintable)
+class UNearestNeighborKMeansData : public UObject
 {
 	GENERATED_BODY()
 public:
@@ -80,4 +80,4 @@ namespace UE::NearestNeighborModel
 		virtual void InitData(UObject& Data, UE::MLDeformer::FMLDeformerEditorToolkit& Toolkit) override;
 		virtual TSharedRef<SWidget> CreateAdditionalWidgets(UObject& Data, TWeakPtr<UE::MLDeformer::FMLDeformerEditorModel> InEditorModel) override;
 	};
-};
\ No newline at end of file
+};
diff --git a/Engine/Plugins/Animation/MLDeformer/NearestNeighborModel/Source/NearestNeighborModelEditor/Private/Tools/NearestNeighborStatsTool.h b/Engine/Plugins/Animation/MLDeformer/NearestNeighborModel/Source/NearestNeighborModelEditor/Private/Tools/NearestNeighborStatsTool.h
--- a/Engine/Plugins/Animation/MLDeformer/NearestNeighborModel/Source/NearestNeighborModelEditor/Private/Tools/NearestNeighborStatsTool.h
+++ b/Engine/Plugins/Animation/MLDeformer/NearestNeighborModel/Source/NearestNeighborModelEditor/Private/Tools/NearestNeighborStatsTool.h
@@ -15,8 +15,8 @@ namespace UE::MLDeformer
 	class FMLDeformerEditorToolkit;
 };
 
-UCLASS(Blueprintable)
-class NEARESTNEIGHBORMODELEDITOR_API UNearestNeighborStatsData : public UObject
+UCLASS(MinimalAPI, Blueprintable)
+class UNearestNeighborStatsData : public UObject
 {
 	GENERATED_BODY()
 public:
@@ -42,4 +42,4 @@ namespace UE::NearestNeighborModel
 		virtual void InitData(UObject& Data, UE::MLDeformer::FMLDeformerEditorToolkit& Toolkit) override;
 		virtual TSharedRef<SWidget> CreateAdditionalWidgets(UObject& Data, TWeakPtr<UE::MLDeformer::FMLDeformerEditorModel> InEditorModel) override;
 	};
-};
\ No newline at end of file
+};
diff --git a/Engine/Plugins/Animation/MLDeformer/NearestNeighborModel/Source/NearestNeighborModelEditor/Public/NearestNeighborEditorModel.h b/Engine/Plugins/Animation/MLDeformer/NearestNeighborModel/Source/NearestNeighborModelEditor/Public/NearestNeighborEditorModel.h
--- a/Engine/Plugins/Animation/MLDeformer/NearestNeighborModel/Source/NearestNeighborModelEditor/Public/NearestNeighborEditorModel.h
+++ b/Engine/Plugins/Animation/MLDeformer/NearestNeighborModel/Source/NearestNeighborModelEditor/Public/NearestNeighborEditorModel.h
@@ -8,6 +8,8 @@
 #include "NearestNeighborModelHelpers.h"
 #include "Types/SlateEnums.h"
 
+#define UE_API NEARESTNEIGHBORMODELEDITOR_API
+
 class UMLDeformerComponent;
 class UNearestNeighborModelSection;
 class USkeletalMesh;
@@ -21,7 +23,7 @@ namespace UE::NearestNeighborModel
 	class FNearestNeighborEditorModelActor;
 	class FVertexMapSelector;
 
-	class NEARESTNEIGHBORMODELEDITOR_API FNearestNeighborEditorModel
+	class FNearestNeighborEditorModel
 		: public UE::MLDeformer::FMLDeformerMorphModelEditorModel
 	{
 	public:
@@ -29,62 +31,62 @@ namespace UE::NearestNeighborModel
 		using FMLDeformerSampler = UE::MLDeformer::FMLDeformerSampler;
 
 		// We need to implement this static MakeInstance method.
-		static FMLDeformerEditorModel* MakeInstance();
+		static UE_API FMLDeformerEditorModel* MakeInstance();
 
 		// FGCObject overrides.
 		virtual FString GetReferencerName() const override { return TEXT("FNearestNeighborEditorModel"); }
 		// ~END FGCObject overrides.
 
 		// FMLDeformerEditorModel overrides.
-		virtual void Init(const InitSettings& Settings) override;
-		virtual TSharedPtr<FMLDeformerSampler> CreateSamplerObject() const override;
-		virtual void CreateActors(const TSharedRef<IPersonaPreviewScene>& InPersonaPreviewScene) override;
-		virtual void Tick(FEditorViewportClient* ViewportClient, float DeltaTime) override;
-		virtual void InitInputInfo(UMLDeformerInputInfo* InputInfo) override;
-		virtual ETrainingResult Train() override;
-		virtual bool LoadTrainedNetwork() const override;
-		virtual FMLDeformerTrainingInputAnim* GetTrainingInputAnim(int32 Index) const override;
-		virtual void UpdateTimelineTrainingAnimList() override;
-		virtual void OnPropertyChanged(FPropertyChangedEvent& PropertyChangedEvent) override;
-		virtual void OnPostTraining(ETrainingResult TrainingResult, bool bUsePartiallyTrainedWhenAborted) override;
-		virtual void Render(const FSceneView* View, FViewport* Viewport, FPrimitiveDrawInterface* PDI) override;
-		virtual void UpdateTrainingDeviceList() override;
-		virtual void CopyBaseSettingsFromModel(const FMLDeformerEditorModel* SourceEditorModel) override;
+		UE_API virtual void Init(const InitSettings& Settings) override;
+		UE_API virtual TSharedPtr<FMLDeformerSampler> CreateSamplerObject() const override;
+		UE_API virtual void CreateActors(const TSharedRef<IPersonaPreviewScene>& InPersonaPreviewScene) override;
+		UE_API virtual void Tick(FEditorViewportClient* ViewportClient, float DeltaTime) override;
+		UE_API virtual void InitInputInfo(UMLDeformerInputInfo* InputInfo) override;
+		UE_API virtual ETrainingResult Train() override;
+		UE_API virtual bool LoadTrainedNetwork() const override;
+		UE_API virtual FMLDeformerTrainingInputAnim* GetTrainingInputAnim(int32 Index) const override;
+		UE_API virtual void UpdateTimelineTrainingAnimList() override;
+		UE_API virtual void OnPropertyChanged(FPropertyChangedEvent& PropertyChangedEvent) override;
+		UE_API virtual void OnPostTraining(ETrainingResult TrainingResult, bool bUsePartiallyTrainedWhenAborted) override;
+		UE_API virtual void Render(const FSceneView* View, FViewport* Viewport, FPrimitiveDrawInterface* PDI) override;
+		UE_API virtual void UpdateTrainingDeviceList() override;
+		UE_API virtual void CopyBaseSettingsFromModel(const FMLDeformerEditorModel* SourceEditorModel) override;
 		// ~END FMLDeformerEditorModel overrides.
 		
 		// UMLDeformerMorphModelEditorModel overrides.
 		virtual bool IsMorphWeightClampingSupported() const override	{ return false; }	// We already do input clamping, so output clamping really isn't needed.
 		// ~END UMLDeformerMorphModelEditorModel overrides.
 
-		void OnUpdateClicked();
-		void ClearReferences();
+		UE_API void OnUpdateClicked();
+		UE_API void ClearReferences();
 
-		FVertexMapSelector* GetVertexMapSelector() const;
-		FVertVizSelector* GetVertVizSelector() const;
+		UE_API FVertexMapSelector* GetVertexMapSelector() const;
+		UE_API FVertVizSelector* GetVertVizSelector() const;
 
 	protected:
-		virtual void CreateSamplers() override;
-		virtual bool IsAnimIndexValid(int32 Index) const override; 
+		UE_API virtual void CreateSamplers() override;
+		UE_API virtual bool IsAnimIndexValid(int32 Index) const override; 
 
 	private:
 		static constexpr int32 DefaultState = INDEX_NONE; 
 
 		// Some helpers that cast to this model's variants of some classes.
-		UNearestNeighborModel* GetCastModel() const;
-		UNearestNeighborModelVizSettings* GetCastVizSettings() const;
-		UMLDeformerComponent* GetTestMLDeformerComponent() const;
-		USkeletalMeshComponent* GetTestSkeletalMeshComponent() const;
-		UNearestNeighborModelInstance* GetTestNearestNeighborModelInstance() const;
-		void CreateDefaultSection();
-
-		EOpFlag Update();
-		EOpFlag CheckNetwork();
-		EOpFlag UpdateNearestNeighborData();
-		EOpFlag UpdateMorphDeltas();
-		void UpdateNearestNeighborIds();
-
-		FNearestNeighborEditorModelActor* CreateNearestNeighborActor(UWorld* World) const;
-		void UpdateNearestNeighborActor(FNearestNeighborEditorModelActor& Actor) const;
+		UE_API UNearestNeighborModel* GetCastModel() const;
+		UE_API UNearestNeighborModelVizSettings* GetCastVizSettings() const;
+		UE_API UMLDeformerComponent* GetTestMLDeformerComponent() const;
+		UE_API USkeletalMeshComponent* GetTestSkeletalMeshComponent() const;
+		UE_API UNearestNeighborModelInstance* GetTestNearestNeighborModelInstance() const;
+		UE_API void CreateDefaultSection();
+
+		UE_API EOpFlag Update();
+		UE_API EOpFlag CheckNetwork();
+		UE_API EOpFlag UpdateNearestNeighborData();
+		UE_API EOpFlag UpdateMorphDeltas();
+		UE_API void UpdateNearestNeighborIds();
+
+		UE_API FNearestNeighborEditorModelActor* CreateNearestNeighborActor(UWorld* World) const;
+		UE_API void UpdateNearestNeighborActor(FNearestNeighborEditorModelActor& Actor) const;
 	
 		FNearestNeighborEditorModelActor* NearestNeighborActor = nullptr;	// This should be only set in CreateActors() and is automatically deleted by the base class.
 		TUniquePtr<FVertexMapSelector> VertexMapSelector;
@@ -124,3 +126,5 @@ namespace UE::NearestNeighborModel
 		TArray<TSharedPtr<FString>> Options;
 	};
 }	// namespace UE::NearestNeighborModel
+
+#undef UE_API
diff --git a/Engine/Plugins/Animation/MLDeformer/NearestNeighborModel/Source/NearestNeighborModelEditor/Public/NearestNeighborGeomCacheSampler.h b/Engine/Plugins/Animation/MLDeformer/NearestNeighborModel/Source/NearestNeighborModelEditor/Public/NearestNeighborGeomCacheSampler.h
--- a/Engine/Plugins/Animation/MLDeformer/NearestNeighborModel/Source/NearestNeighborModelEditor/Public/NearestNeighborGeomCacheSampler.h
+++ b/Engine/Plugins/Animation/MLDeformer/NearestNeighborModel/Source/NearestNeighborModelEditor/Public/NearestNeighborGeomCacheSampler.h
@@ -5,6 +5,8 @@
 #include "MLDeformerGeomCacheSampler.h"
 #include "NearestNeighborModelHelpers.h"
 
+#define UE_API NEARESTNEIGHBORMODELEDITOR_API
+
 class UAnimSequence;
 class UGeometryCache;
 class UNearestNeighborTrainingModel;
@@ -13,12 +15,12 @@ namespace UE::NearestNeighborModel
 {
 	class FNearestNeighborEditorModel;
 
-	class NEARESTNEIGHBORMODELEDITOR_API FNearestNeighborGeomCacheSampler
+	class FNearestNeighborGeomCacheSampler
 		: public UE::MLDeformer::FMLDeformerGeomCacheSampler
 	{
 	public:
 		// FMLDeformerGeomCacheSampler overrides
-		virtual void Sample(int32 InAnimFrameIndex) override;
+		UE_API virtual void Sample(int32 InAnimFrameIndex) override;
 		// ~END FMLDeformerGeomCacheSampler overrides
 
 		/** This will set the SkeletalMeshComponent with custom Anim and GeometryCacheComponent with custom Cache. 
@@ -26,25 +28,27 @@ namespace UE::NearestNeighborModel
 		 * Using this function will break Sample(). Call CustomSample() instead.
 		 * Do NOT use this function on EditorModel Samplers. 
 		 */ 
-		void Customize(UAnimSequence* Anim, UGeometryCache* Cache = nullptr);
+		UE_API void Customize(UAnimSequence* Anim, UGeometryCache* Cache = nullptr);
 
 		/**
 		 * Sample function used with Customize
 		 * @return true if sampling is successful
 		 */
-		bool CustomSample(int32 Frame);
+		UE_API bool CustomSample(int32 Frame);
 
 		/**
 		 * Get the Mesh Index Buffer of the skeletal mesh object used by the model.
 		 * 
 		 * @return A flattened array of the index buffer.
 		 */
-		TArray<uint32> GetMeshIndexBuffer() const;
+		UE_API TArray<uint32> GetMeshIndexBuffer() const;
 
 	private:
-		void SampleDualQuaternionDeltas(int32 InAnimFrameIndex);
-		EOpFlag GenerateMeshMappings();
-		EOpFlag CheckMeshMappingsEmpty() const;
-		FVector3f CalcDualQuaternionDelta(int32 VertexIndex, const FVector3f& WorldDelta, const FSkeletalMeshLODRenderData& SkelMeshLODData, const FSkinWeightVertexBuffer& SkinWeightBuffer) const;
+		UE_API void SampleDualQuaternionDeltas(int32 InAnimFrameIndex);
+		UE_API EOpFlag GenerateMeshMappings();
+		UE_API EOpFlag CheckMeshMappingsEmpty() const;
+		UE_API FVector3f CalcDualQuaternionDelta(int32 VertexIndex, const FVector3f& WorldDelta, const FSkeletalMeshLODRenderData& SkelMeshLODData, const FSkinWeightVertexBuffer& SkinWeightBuffer) const;
 	};
 }
+
+#undef UE_API
diff --git a/Engine/Plugins/Animation/MLDeformer/NearestNeighborModel/Source/NearestNeighborModelEditor/Public/NearestNeighborModelDetails.h b/Engine/Plugins/Animation/MLDeformer/NearestNeighborModel/Source/NearestNeighborModelEditor/Public/NearestNeighborModelDetails.h
--- a/Engine/Plugins/Animation/MLDeformer/NearestNeighborModel/Source/NearestNeighborModelEditor/Public/NearestNeighborModelDetails.h
+++ b/Engine/Plugins/Animation/MLDeformer/NearestNeighborModel/Source/NearestNeighborModelEditor/Public/NearestNeighborModelDetails.h
@@ -5,6 +5,8 @@
 #include "CoreTypes.h"
 #include "MLDeformerMorphModelDetails.h"
 
+#define UE_API NEARESTNEIGHBORMODELEDITOR_API
+
 class IDetailChildrenBuilder;
 class IPropertyHandle;
 class UNearestNeighborModel;
@@ -13,37 +15,39 @@ namespace UE::NearestNeighborModel
 {
 	class FNearestNeighborEditorModel;
 
-	class NEARESTNEIGHBORMODELEDITOR_API FNearestNeighborModelDetails
+	class FNearestNeighborModelDetails
 		: public ::UE::MLDeformer::FMLDeformerMorphModelDetails
 	{
 	public:
 		/** Makes a new instance of this detail layout class for a specific detail view requesting it. */
-		static TSharedRef<IDetailCustomization> MakeInstance();
+		static UE_API TSharedRef<IDetailCustomization> MakeInstance();
 
 		// ILayoutDetails overrides.
-		virtual void CustomizeDetails(IDetailLayoutBuilder& DetailBuilder) override;
-		virtual void CreateCategories() override;
+		UE_API virtual void CustomizeDetails(IDetailLayoutBuilder& DetailBuilder) override;
+		UE_API virtual void CreateCategories() override;
 		// ~END ILayoutDetails overrides.
 
 		// FMLDeformerModelDetails overrides.
-		virtual bool UpdateMemberPointers(const TArray<TWeakObjectPtr<UObject>>& Objects) override;
+		UE_API virtual bool UpdateMemberPointers(const TArray<TWeakObjectPtr<UObject>>& Objects) override;
 		// ~END FMLDeformerModelDetails overrides.
 
 	private:
-		UNearestNeighborModel* GetCastModel() const;
-		FNearestNeighborEditorModel* GetCastEditorModel() const;
+		UE_API UNearestNeighborModel* GetCastModel() const;
+		UE_API FNearestNeighborEditorModel* GetCastEditorModel() const;
 
-		void CustomizeTrainingSettingsCategory() const;
-		void CustomizeNearestNeighborSettingsCategory() const;
-		void CustomizeSectionsCategory(IDetailLayoutBuilder& DetailBuilder);
-		void CustomizeMorphTargetCategory(IDetailLayoutBuilder& DetailBuilder) const;
-		void CustomizeStatusCategory() const;
-		void CustomizeFileCacheCategory(IDetailLayoutBuilder& DetailBuilder) const;
+		UE_API void CustomizeTrainingSettingsCategory() const;
+		UE_API void CustomizeNearestNeighborSettingsCategory() const;
+		UE_API void CustomizeSectionsCategory(IDetailLayoutBuilder& DetailBuilder);
+		UE_API void CustomizeMorphTargetCategory(IDetailLayoutBuilder& DetailBuilder) const;
+		UE_API void CustomizeStatusCategory() const;
+		UE_API void CustomizeFileCacheCategory(IDetailLayoutBuilder& DetailBuilder) const;
 
-		void GenerateSectionElementWidget(TSharedRef<IPropertyHandle> PropertyHandle, int32 ArrayIndex, IDetailChildrenBuilder& ChildrenBuilder);
+		UE_API void GenerateSectionElementWidget(TSharedRef<IPropertyHandle> PropertyHandle, int32 ArrayIndex, IDetailChildrenBuilder& ChildrenBuilder);
 
 		IDetailCategoryBuilder* NearestNeighborCategoryBuilder = nullptr;
 		IDetailCategoryBuilder* StatusCategoryBuilder = nullptr;
 		IDetailCategoryBuilder* SectionsCategoryBuilder = nullptr;
 	};
 }	// namespace UE::NearestNeighborModel
+
+#undef UE_API
diff --git a/Engine/Plugins/Animation/MLDeformer/NearestNeighborModel/Source/NearestNeighborModelEditor/Public/NearestNeighborModelStyle.h b/Engine/Plugins/Animation/MLDeformer/NearestNeighborModel/Source/NearestNeighborModelEditor/Public/NearestNeighborModelStyle.h
--- a/Engine/Plugins/Animation/MLDeformer/NearestNeighborModel/Source/NearestNeighborModelEditor/Public/NearestNeighborModelStyle.h
+++ b/Engine/Plugins/Animation/MLDeformer/NearestNeighborModel/Source/NearestNeighborModelEditor/Public/NearestNeighborModelStyle.h
@@ -4,16 +4,20 @@
 
 #include "Styling/SlateStyle.h"
 
+#define UE_API NEARESTNEIGHBORMODELEDITOR_API
+
 namespace UE::NearestNeighborModel
 {
-	class NEARESTNEIGHBORMODELEDITOR_API FNearestNeighborModelEditorStyle
+	class FNearestNeighborModelEditorStyle
 		: public FSlateStyleSet
 	{
 	public:
-		FNearestNeighborModelEditorStyle();
-		virtual ~FNearestNeighborModelEditorStyle();
+		UE_API FNearestNeighborModelEditorStyle();
+		UE_API virtual ~FNearestNeighborModelEditorStyle();
 
-		static FNearestNeighborModelEditorStyle& Get();
+		static UE_API FNearestNeighborModelEditorStyle& Get();
 	};
 
 }	// namespace UE::NearestNeighborModel
+
+#undef UE_API
diff --git a/Engine/Plugins/Animation/MLDeformer/NearestNeighborModel/Source/NearestNeighborModelEditor/Public/NearestNeighborModelVizSettingsDetails.h b/Engine/Plugins/Animation/MLDeformer/NearestNeighborModel/Source/NearestNeighborModelEditor/Public/NearestNeighborModelVizSettingsDetails.h
--- a/Engine/Plugins/Animation/MLDeformer/NearestNeighborModel/Source/NearestNeighborModelEditor/Public/NearestNeighborModelVizSettingsDetails.h
+++ b/Engine/Plugins/Animation/MLDeformer/NearestNeighborModel/Source/NearestNeighborModelEditor/Public/NearestNeighborModelVizSettingsDetails.h
@@ -4,20 +4,24 @@
 
 #include "MLDeformerMorphModelVizSettingsDetails.h"
 
+#define UE_API NEARESTNEIGHBORMODELEDITOR_API
+
 namespace UE::NearestNeighborModel
 {
 	class FNearestNeighborEditorModel;
-	class NEARESTNEIGHBORMODELEDITOR_API FNearestNeighborModelVizSettingsDetails
+	class FNearestNeighborModelVizSettingsDetails
 		: public UE::MLDeformer::FMLDeformerMorphModelVizSettingsDetails
 	{
 	public:
 		/** Makes a new instance of this detail layout class for a specific detail view requesting it. */
-		static TSharedRef<IDetailCustomization> MakeInstance();
-		virtual void AddAdditionalSettings() override;
+		static UE_API TSharedRef<IDetailCustomization> MakeInstance();
+		UE_API virtual void AddAdditionalSettings() override;
 
 	private:
-		void AddTrainingMeshAdditionalSettings();
-		void AddLiveAdditionalSettings();
-		FNearestNeighborEditorModel* GetCastEditorModel() const;
+		UE_API void AddTrainingMeshAdditionalSettings();
+		UE_API void AddLiveAdditionalSettings();
+		UE_API FNearestNeighborEditorModel* GetCastEditorModel() const;
 	};
 }	// namespace UE::NearestNeighborModel
+
+#undef UE_API
diff --git a/Engine/Plugins/Animation/MLDeformer/NearestNeighborModel/Source/NearestNeighborModelEditor/Public/NearestNeighborTrainingModel.h b/Engine/Plugins/Animation/MLDeformer/NearestNeighborModel/Source/NearestNeighborModelEditor/Public/NearestNeighborTrainingModel.h
--- a/Engine/Plugins/Animation/MLDeformer/NearestNeighborModel/Source/NearestNeighborModelEditor/Public/NearestNeighborTrainingModel.h
+++ b/Engine/Plugins/Animation/MLDeformer/NearestNeighborModel/Source/NearestNeighborModelEditor/Public/NearestNeighborTrainingModel.h
@@ -7,6 +7,8 @@
 
 #include "NearestNeighborTrainingModel.generated.h"
 
+#define UE_API NEARESTNEIGHBORMODELEDITOR_API
+
 namespace UE::NearestNeighborModel
 {
 	class FNearestNeighborEditorModel;
@@ -19,27 +21,27 @@ class UNearestNeighborModel;
 class UNearestNeighborStatsData;
 class UNearestNeighborModelInstance;
 
-UCLASS(Blueprintable)
-class NEARESTNEIGHBORMODELEDITOR_API UNearestNeighborTrainingModel
+UCLASS(MinimalAPI, Blueprintable)
+class UNearestNeighborTrainingModel
 	: public UMLDeformerGeomCacheTrainingModel
 {
 	GENERATED_BODY()
 
 public:
-	virtual ~UNearestNeighborTrainingModel();
-	virtual void Init(UE::MLDeformer::FMLDeformerEditorModel* InEditorModel) override;
+	UE_API virtual ~UNearestNeighborTrainingModel();
+	UE_API virtual void Init(UE::MLDeformer::FMLDeformerEditorModel* InEditorModel) override;
 
 	UFUNCTION(BlueprintImplementableEvent, Category = "Python")
-	int32 Train() const;
+	UE_API int32 Train() const;
 
 	UFUNCTION(BlueprintImplementableEvent, Category = "Python")
-	int32 UpdateNearestNeighborData();
+	UE_API int32 UpdateNearestNeighborData();
 
 	UFUNCTION(BlueprintImplementableEvent, Category = "Python")
-	int32 KmeansClusterPoses(const UNearestNeighborKMeansData* KMeansData);
+	UE_API int32 KmeansClusterPoses(const UNearestNeighborKMeansData* KMeansData);
 
 	UFUNCTION(BlueprintImplementableEvent, Category = "Python")
-	bool GetNeighborStats(const UNearestNeighborStatsData* StatsData);
+	UE_API bool GetNeighborStats(const UNearestNeighborStatsData* StatsData);
 
 protected:
 	UPROPERTY(VisibleAnywhere, BlueprintReadOnly, Category = "Python")
@@ -50,41 +52,43 @@ protected:
 
 private:
 	UFUNCTION(BlueprintPure, Category = "Python")
-	int32 GetNumFramesAnimSequence(const UAnimSequence* Anim) const;
+	UE_API int32 GetNumFramesAnimSequence(const UAnimSequence* Anim) const;
 
 	UFUNCTION(BlueprintPure, Category = "Python")
-	int32 GetNumFramesGeometryCache(const UGeometryCache* GeometryCache) const;
+	UE_API int32 GetNumFramesGeometryCache(const UGeometryCache* GeometryCache) const;
 	
 	UFUNCTION(BlueprintPure, Category = "Python")
-	const USkeleton* GetModelSkeleton(const UMLDeformerModel* Model) const;
+	UE_API const USkeleton* GetModelSkeleton(const UMLDeformerModel* Model) const;
 	
 	UFUNCTION(BlueprintPure, Category = "Python")
-	UNearestNeighborModel* GetNearestNeighborModel() const;
+	UE_API UNearestNeighborModel* GetNearestNeighborModel() const;
 
 	UFUNCTION(BlueprintCallable, Category = "Python")
-	bool SetCustomSamplerData(UAnimSequence* Anim, UGeometryCache* Cache = nullptr);
+	UE_API bool SetCustomSamplerData(UAnimSequence* Anim, UGeometryCache* Cache = nullptr);
 
 	UFUNCTION(BlueprintCallable, Category = "Python")
-	bool CustomSample(int32 Frame);
+	UE_API bool CustomSample(int32 Frame);
 
 	UFUNCTION(BlueprintCallable, Category = "Python")
-	bool SetCustomSamplerDataFromSection(int32 SectionIndex);
+	UE_API bool SetCustomSamplerDataFromSection(int32 SectionIndex);
 
 	UFUNCTION(BlueprintPure, Category = "Python")
-	TArray<float> GetUnskinnedVertexPositions();
+	UE_API TArray<float> GetUnskinnedVertexPositions();
 
 	UFUNCTION(BlueprintPure, Category = "Python")
-	TArray<int32> GetMeshIndexBuffer();
+	UE_API TArray<int32> GetMeshIndexBuffer();
 
 	UFUNCTION(BlueprintCallable, Category = "Python")
-	UNearestNeighborModelInstance* CreateModelInstance();
+	UE_API UNearestNeighborModelInstance* CreateModelInstance();
 
 	UFUNCTION(BlueprintCallable, Category = "Python")
-	void DestroyModelInstance(UNearestNeighborModelInstance* ModelInstance);
+	UE_API void DestroyModelInstance(UNearestNeighborModelInstance* ModelInstance);
 
-	void SetNewCustomSampler();
-	UNearestNeighborModel* GetCastModel() const;
-	UE::NearestNeighborModel::FNearestNeighborEditorModel* GetCastEditorModel() const;
+	UE_API void SetNewCustomSampler();
+	UE_API UNearestNeighborModel* GetCastModel() const;
+	UE_API UE::NearestNeighborModel::FNearestNeighborEditorModel* GetCastEditorModel() const;
 
 	TUniquePtr<UE::NearestNeighborModel::FNearestNeighborGeomCacheSampler> CustomSampler;
 };
+
+#undef UE_API
diff --git a/Engine/Plugins/Animation/MLDeformer/NeuralMorphModel/Source/NeuralMorphModelEditor/Public/NeuralMorphEditorModel.h b/Engine/Plugins/Animation/MLDeformer/NeuralMorphModel/Source/NeuralMorphModelEditor/Public/NeuralMorphEditorModel.h
--- a/Engine/Plugins/Animation/MLDeformer/NeuralMorphModel/Source/NeuralMorphModelEditor/Public/NeuralMorphEditorModel.h
+++ b/Engine/Plugins/Animation/MLDeformer/NeuralMorphModel/Source/NeuralMorphModelEditor/Public/NeuralMorphEditorModel.h
@@ -6,6 +6,8 @@
 #include "MLDeformerMorphModelEditorModel.h"
 #include "NeuralMorphModel.h"
 
+#define UE_API NEURALMORPHMODELEDITOR_API
+
 class UNeuralMorphModel;
 
 namespace UE::MLDeformer
@@ -20,35 +22,35 @@ namespace UE::NeuralMorphModel
 	/**
 	 * The editor model related to the neural morph model's runtime class (UNeuralMorphModel).
 	 */
-	class NEURALMORPHMODELEDITOR_API FNeuralMorphEditorModel
+	class FNeuralMorphEditorModel
 		: public FMLDeformerMorphModelEditorModel
 	{
 	public:
 		// We need to implement this static MakeInstance method.
-		static FMLDeformerEditorModel* MakeInstance();
+		static UE_API FMLDeformerEditorModel* MakeInstance();
 
 		// FGCObject overrides.
 		virtual FString GetReferencerName() const override		{ return TEXT("FNeuralMorphEditorModel"); }
 		// ~END FGCObject overrides.
 	
 		// FMLDeformerEditorModel overrides.
-		virtual void Init(const InitSettings& Settings) override;
-		virtual ETrainingResult Train() override;
-		virtual void InitInputInfo(UMLDeformerInputInfo* InputInfo) override;
-		virtual FText GetOverlayText() const override;
-		virtual void OnPropertyChanged(FPropertyChangedEvent& PropertyChangedEvent) override;
-		virtual TSharedPtr<SMLDeformerInputWidget> CreateInputWidget() override;
-		virtual void OnPostInputAssetChanged() override;
-		virtual void Render(const FSceneView* View, FViewport* Viewport, FPrimitiveDrawInterface* PDI) override;
+		UE_API virtual void Init(const InitSettings& Settings) override;
+		UE_API virtual ETrainingResult Train() override;
+		UE_API virtual void InitInputInfo(UMLDeformerInputInfo* InputInfo) override;
+		UE_API virtual FText GetOverlayText() const override;
+		UE_API virtual void OnPropertyChanged(FPropertyChangedEvent& PropertyChangedEvent) override;
+		UE_API virtual TSharedPtr<SMLDeformerInputWidget> CreateInputWidget() override;
+		UE_API virtual void OnPostInputAssetChanged() override;
+		UE_API virtual void Render(const FSceneView* View, FViewport* Viewport, FPrimitiveDrawInterface* PDI) override;
 		virtual bool GetSupportsPerTrainingInputAnimVertexMask() const override	{ return true; }
-		virtual void UpdateTrainingDeviceList() override;
-		virtual UMeshDeformer* GetActiveHeatMapDeformer() const override;
-		virtual void UpdateIsReadyForTrainingState() override;
-		virtual bool LoadTrainedNetwork() const override;
+		UE_API virtual void UpdateTrainingDeviceList() override;
+		UE_API virtual UMeshDeformer* GetActiveHeatMapDeformer() const override;
+		UE_API virtual void UpdateIsReadyForTrainingState() override;
+		UE_API virtual bool LoadTrainedNetwork() const override;
 		// ~END FMLDeformerEditorModel overrides.
 
 		// FMLDeformerMorphModelEditorModel overrides.
-		virtual const TArrayView<const float> GetMaskForMorphTarget(int32 MorphTargetIndex) const override;
+		UE_API virtual const TArrayView<const float> GetMaskForMorphTarget(int32 MorphTargetIndex) const override;
 		virtual bool IsInputMaskingSupported() const override	{ return true; }
 		// ~FMLDeformerMorphModelEditorModel overrides.
 
@@ -64,33 +66,33 @@ namespace UE::NeuralMorphModel
 		 * This will modify the item mask buffer as returned by UMLDeformerMorphModel::GetInputItemMaskBuffer().
 		 * @param OutMaskBuffer The mask buffer to write to. This will automatically be resized.
 		 */
-		void BuildMaskBuffer(TArray<float>& OutMaskBuffer);
+		UE_API void BuildMaskBuffer(TArray<float>& OutMaskBuffer);
 
 		/** Rebuild the mask buffer. */
-		void RebuildEditorMaskInfo();
+		UE_API void RebuildEditorMaskInfo();
 
 		/**
 		 * Generate a unique name for a bone group.
 		 * @return The unique name.
 		 */
-		FName GenerateUniqueBoneGroupName() const;
+		UE_API FName GenerateUniqueBoneGroupName() const;
 
 		/**
 		 * Generate a unique name for a curve group.
 		 * @return The unique name.
 		 */
-		FName GenerateUniqueCurveGroupName() const;
+		UE_API FName GenerateUniqueCurveGroupName() const;
 
 		/**
 		 * Remove mask infos for bones that do not exist in the bone include list.
 		 */
-		void RemoveNonExistingMaskInfos();
+		UE_API void RemoveNonExistingMaskInfos();
 
 		/** Clear all bone mask infos. */
-		void ResetBoneMaskInfos();
+		UE_API void ResetBoneMaskInfos();
 
 		/** Clear all bone group mask infos. */
-		void ResetBoneGroupMaskInfos();
+		UE_API void ResetBoneGroupMaskInfos();
 
 		/**
 		 * Generate bone mask infos for all bones in the bone input list.
@@ -102,7 +104,7 @@ namespace UE::NeuralMorphModel
 		 * However we don't want the mask to be too large either, as that costs performance and memory usage.
 		 * @param HierarchyDepth The hierarchy depth to use, which should be a value of 1 or above.
 		 */
-		void GenerateBoneMaskInfos(int32 HierarchyDepth);
+		UE_API void GenerateBoneMaskInfos(int32 HierarchyDepth);
 
 		/**
 		 * Generate the mask infos for all bone groups. 
@@ -110,7 +112,7 @@ namespace UE::NeuralMorphModel
 		 * @param HierarchyDepth The hierarchy depth to use, which should be a value of 1 or above.
 		 * @see GenerateBoneMaskInfos
 		 */
-		void GenerateBoneGroupMaskInfos(int32 HierarchyDepth);
+		UE_API void GenerateBoneGroupMaskInfos(int32 HierarchyDepth);
 
 		/**
 		 * Generate the mask info for a given bone.
@@ -119,7 +121,7 @@ namespace UE::NeuralMorphModel
 		 * @param HierarchyDepth The hierarchy depth to use, which should be a value of 1 or above.
 		 * @see GenerateBoneMaskInfos
 		 */
-		void GenerateBoneMaskInfo(int32 InputInfoBoneIndex, int32 HierarchyDepth);
+		UE_API void GenerateBoneMaskInfo(int32 InputInfoBoneIndex, int32 HierarchyDepth);
 
 		/**
 		 * Generate the mask info for a given bone group.
@@ -128,7 +130,7 @@ namespace UE::NeuralMorphModel
 		 * @param HierarchyDepth The hierarchy depth to use, which should be a value of 1 or above.
 		 * @see GenerateBoneMaskInfos
 		 */
-		void GenerateBoneGroupMaskInfo(int32 InputInfoBoneGroupIndex, int32 HierarchyDepth);
+		UE_API void GenerateBoneGroupMaskInfo(int32 InputInfoBoneGroupIndex, int32 HierarchyDepth);
 
 		PRAGMA_DISABLE_DEPRECATION_WARNINGS;
 		UE_DEPRECATED(5.5, "Please use the FMLDeformerMorphModelEditorModel::ApplyMaskInfoToBuffer that takes an FMLDeformerMaskInfo as mask info type.")
@@ -158,7 +160,7 @@ namespace UE::NeuralMorphModel
 		 *        So if you have 10 bones, and 5 curves, and the MaskItemIndex is value 12, it will draw the mask for the third curve, because items [0..9] would be for the bones, and [10..14] for the curves.
 		 * @param DrawOffset The offset in world space units to draw the mask at.
 		 */
-		void DebugDrawItemMask(FPrimitiveDrawInterface* PDI, int32 MaskItemIndex, const FVector& DrawOffset);
+		UE_API void DebugDrawItemMask(FPrimitiveDrawInterface* PDI, int32 MaskItemIndex, const FVector& DrawOffset);
 
 	private:
 		/**
@@ -170,9 +172,11 @@ namespace UE::NeuralMorphModel
 		 * @param RefSkel The reference skeleton we use to find child nodes.
 		 * @param SkelBoneIndices The input and output array of indices inside the reference skeleton. This method potentially adds new entries to this array.
 		 */
-		void AddTwistBones(const FReferenceSkeleton& RefSkel, TArray<int32>& SkelBoneIndices);
+		UE_API void AddTwistBones(const FReferenceSkeleton& RefSkel, TArray<int32>& SkelBoneIndices);
 
 	protected:
 		int32 MaskVizItemIndex = INDEX_NONE;
 	};
 }	// namespace UE::NeuralMorphModel
+
+#undef UE_API
diff --git a/Engine/Plugins/Animation/MLDeformer/NeuralMorphModel/Source/NeuralMorphModelEditor/Public/NeuralMorphEditorProjectSettings.h b/Engine/Plugins/Animation/MLDeformer/NeuralMorphModel/Source/NeuralMorphModelEditor/Public/NeuralMorphEditorProjectSettings.h
--- a/Engine/Plugins/Animation/MLDeformer/NeuralMorphModel/Source/NeuralMorphModelEditor/Public/NeuralMorphEditorProjectSettings.h
+++ b/Engine/Plugins/Animation/MLDeformer/NeuralMorphModel/Source/NeuralMorphModelEditor/Public/NeuralMorphEditorProjectSettings.h
@@ -6,23 +6,25 @@
 #include "Engine/DeveloperSettings.h"
 #include "NeuralMorphEditorProjectSettings.generated.h"
 
+#define UE_API NEURALMORPHMODELEDITOR_API
 
-UCLASS(config = EditorPerProjectUserSettings)
-class NEURALMORPHMODELEDITOR_API UNeuralMorphEditorProjectSettings
+
+UCLASS(MinimalAPI, config = EditorPerProjectUserSettings)
+class UNeuralMorphEditorProjectSettings
 	: public UDeveloperSettings
 {
 	GENERATED_BODY()
 
 public:
-	static const UNeuralMorphEditorProjectSettings* Get();
+	static UE_API const UNeuralMorphEditorProjectSettings* Get();
 
 	// UDeveloperSettings overrides
 	virtual FName GetContainerName() const	{ return FName("Project"); }
 	virtual FName GetCategoryName() const	{ return FName("Plugins"); }
 	virtual FName GetSectionName() const	{ return FName("Neural Morph Model"); }
 #if WITH_EDITOR
-	virtual FText GetSectionText() const override;
-	virtual FText GetSectionDescription() const override;
+	UE_API virtual FText GetSectionText() const override;
+	UE_API virtual FText GetSectionDescription() const override;
 #endif
 	// ~END UDeveloperSettings overrides
 
@@ -36,3 +38,5 @@ public:
 	UPROPERTY(config, EditAnywhere, Category = "Bone and Bone Group masking")
 	FString TwistBoneFilter = TEXT("twist");
 };
+
+#undef UE_API
diff --git a/Engine/Plugins/Animation/MLDeformer/NeuralMorphModel/Source/NeuralMorphModelEditor/Public/NeuralMorphModelDetails.h b/Engine/Plugins/Animation/MLDeformer/NeuralMorphModel/Source/NeuralMorphModelEditor/Public/NeuralMorphModelDetails.h
--- a/Engine/Plugins/Animation/MLDeformer/NeuralMorphModel/Source/NeuralMorphModelEditor/Public/NeuralMorphModelDetails.h
+++ b/Engine/Plugins/Animation/MLDeformer/NeuralMorphModel/Source/NeuralMorphModelEditor/Public/NeuralMorphModelDetails.h
@@ -6,6 +6,8 @@
 #include "MLDeformerMorphModelDetails.h"
 #include "IDetailCustomization.h"
 
+#define UE_API NEURALMORPHMODELEDITOR_API
+
 class IDetailLayoutBuilder;
 
 namespace UE::NeuralMorphModel
@@ -13,15 +15,17 @@ namespace UE::NeuralMorphModel
 	/**
 	 * The detail customization for the neural morph model.
 	 */
-	class NEURALMORPHMODELEDITOR_API FNeuralMorphModelDetails
+	class FNeuralMorphModelDetails
 		: public UE::MLDeformer::FMLDeformerMorphModelDetails
 	{
 	public:
 		/** Makes a new instance of this detail layout class for a specific detail view requesting it. */
-		static TSharedRef<IDetailCustomization> MakeInstance();
+		static UE_API TSharedRef<IDetailCustomization> MakeInstance();
 
 		// ILayoutDetails overrides.
-		virtual void CustomizeDetails(IDetailLayoutBuilder& DetailBuilder) override;
+		UE_API virtual void CustomizeDetails(IDetailLayoutBuilder& DetailBuilder) override;
 		// ~END ILayoutDetails overrides.
 	};
 }	// namespace UE::NeuralMorphModel
+
+#undef UE_API
diff --git a/Engine/Plugins/Animation/MLDeformer/NeuralMorphModel/Source/NeuralMorphModelEditor/Public/NeuralMorphModelVizSettingsDetails.h b/Engine/Plugins/Animation/MLDeformer/NeuralMorphModel/Source/NeuralMorphModelEditor/Public/NeuralMorphModelVizSettingsDetails.h
--- a/Engine/Plugins/Animation/MLDeformer/NeuralMorphModel/Source/NeuralMorphModelEditor/Public/NeuralMorphModelVizSettingsDetails.h
+++ b/Engine/Plugins/Animation/MLDeformer/NeuralMorphModel/Source/NeuralMorphModelEditor/Public/NeuralMorphModelVizSettingsDetails.h
@@ -13,7 +13,7 @@ namespace UE::NeuralMorphModel
 	 * We do not have to do anything else. We still need this class in order to do our
 	 * detail customization correctly. Just using the Morph Model directly would cause issues.
 	 */
-	class NEURALMORPHMODELEDITOR_API FNeuralMorphModelVizSettingsDetails
+	class FNeuralMorphModelVizSettingsDetails
 		: public UE::MLDeformer::FMLDeformerMorphModelVizSettingsDetails
 	{
 	public:
diff --git a/Engine/Plugins/Animation/MLDeformer/NeuralMorphModel/Source/NeuralMorphModelEditor/Public/NeuralMorphTrainingModel.h b/Engine/Plugins/Animation/MLDeformer/NeuralMorphModel/Source/NeuralMorphModelEditor/Public/NeuralMorphTrainingModel.h
--- a/Engine/Plugins/Animation/MLDeformer/NeuralMorphModel/Source/NeuralMorphModelEditor/Public/NeuralMorphTrainingModel.h
+++ b/Engine/Plugins/Animation/MLDeformer/NeuralMorphModel/Source/NeuralMorphModelEditor/Public/NeuralMorphTrainingModel.h
@@ -7,18 +7,20 @@
 #include "MLDeformerGeomCacheTrainingModel.h"
 #include "NeuralMorphTrainingModel.generated.h"
 
+#define UE_API NEURALMORPHMODELEDITOR_API
+
 /**
  * The training model for the neural morph model.
  * This class is our link to the Python training.
  */
-UCLASS(Blueprintable)
-class NEURALMORPHMODELEDITOR_API UNeuralMorphTrainingModel
+UCLASS(MinimalAPI, Blueprintable)
+class UNeuralMorphTrainingModel
 	: public UMLDeformerGeomCacheTrainingModel
 {
 	GENERATED_BODY()
 
 public:
-	virtual void Init(UE::MLDeformer::FMLDeformerEditorModel* InEditorModel) override;
+	UE_API virtual void Init(UE::MLDeformer::FMLDeformerEditorModel* InEditorModel) override;
 
 	/**
 	 * Main training function, with implementation in python.
@@ -26,20 +28,22 @@ public:
 	 * @see UMLDeformerTrainingModel
 	 */
 	UFUNCTION(BlueprintImplementableEvent, Category = "Training Model")
-	int32 Train() const;
+	UE_API int32 Train() const;
 
 	UFUNCTION(BlueprintPure, Category = "Training Model")
-	int32 GetNumBoneGroups() const;
+	UE_API int32 GetNumBoneGroups() const;
 
 	UFUNCTION(BlueprintPure, Category = "Training Model")
-	int32 GetNumCurveGroups() const;
+	UE_API int32 GetNumCurveGroups() const;
 
 	UFUNCTION(BlueprintPure, Category = "Training Model")
-	TArray<int32> GenerateBoneGroupIndices() const;
+	UE_API TArray<int32> GenerateBoneGroupIndices() const;
 
 	UFUNCTION(BlueprintPure, Category = "Training Model")
-	TArray<int32> GenerateCurveGroupIndices() const;
+	UE_API TArray<int32> GenerateCurveGroupIndices() const;
 
 	UFUNCTION(BlueprintPure, Category = "Training Model")
-	TArray<float> GetMorphTargetMasks() const;
+	UE_API TArray<float> GetMorphTargetMasks() const;
 };
+
+#undef UE_API
diff --git a/Engine/Plugins/Animation/MLDeformer/NeuralMorphModel/Source/NeuralMorphModelEditor/Public/SNeuralMorphBoneGroupsWidget.h b/Engine/Plugins/Animation/MLDeformer/NeuralMorphModel/Source/NeuralMorphModelEditor/Public/SNeuralMorphBoneGroupsWidget.h
--- a/Engine/Plugins/Animation/MLDeformer/NeuralMorphModel/Source/NeuralMorphModelEditor/Public/SNeuralMorphBoneGroupsWidget.h
+++ b/Engine/Plugins/Animation/MLDeformer/NeuralMorphModel/Source/NeuralMorphModelEditor/Public/SNeuralMorphBoneGroupsWidget.h
@@ -8,6 +8,8 @@
 #include "Widgets/Views/STreeView.h"
 #include "Framework/Commands/Commands.h"
 
+#define UE_API NEURALMORPHMODELEDITOR_API
+
 class SWidget;
 
 namespace UE::MLDeformer
@@ -67,7 +69,7 @@ namespace UE::NeuralMorphModel
 	};
 
 
-	class NEURALMORPHMODELEDITOR_API SNeuralMorphBoneGroupsWidget
+	class SNeuralMorphBoneGroupsWidget
 		: public STreeView<TSharedPtr<FNeuralMorphBoneGroupsTreeElement>>
 	{
 	public:
@@ -76,32 +78,32 @@ namespace UE::NeuralMorphModel
 		SLATE_ARGUMENT(TSharedPtr<UE::MLDeformer::SMLDeformerInputWidget>, InputWidget)
 		SLATE_END_ARGS()
 
-		void Construct(const FArguments& InArgs);
-		void BindCommands(TSharedPtr<FUICommandList> CommandList);
-		void Refresh();
+		UE_API void Construct(const FArguments& InArgs);
+		UE_API void BindCommands(TSharedPtr<FUICommandList> CommandList);
+		UE_API void Refresh();
 
-		FText GetSectionTitle() const;
-		int32 GetNumSelectedGroups() const;
-		TSharedPtr<UE::MLDeformer::SMLDeformerInputWidget> GetInputWidget() const;
+		UE_API FText GetSectionTitle() const;
+		UE_API int32 GetNumSelectedGroups() const;
+		UE_API TSharedPtr<UE::MLDeformer::SMLDeformerInputWidget> GetInputWidget() const;
 
 	private:
-		void AddElement(TSharedPtr<FNeuralMorphBoneGroupsTreeElement> Element, TSharedPtr<FNeuralMorphBoneGroupsTreeElement> ParentElement);	
-		void HandleGetChildrenForTree(TSharedPtr<FNeuralMorphBoneGroupsTreeElement> InItem, TArray<TSharedPtr<FNeuralMorphBoneGroupsTreeElement>>& OutChildren);	
-		void OnSelectionChanged(TSharedPtr<FNeuralMorphBoneGroupsTreeElement> Selection, ESelectInfo::Type SelectInfo);
-		const TArray<TSharedPtr<FNeuralMorphBoneGroupsTreeElement>>& GetRootElements() const;
-		FReply OnKeyDown(const FGeometry& InGeometry, const FKeyEvent& InKeyEvent) override;
-		TSharedRef<ITableRow> MakeTableRowWidget(TSharedPtr<FNeuralMorphBoneGroupsTreeElement> InItem, const TSharedRef<STableViewBase>& OwnerTable);
-		TSharedPtr<SWidget> CreateContextMenuWidget() const;
-
-		void UpdateTreeElements();
-		void RefreshTree(bool bBroadcastPropertyChanged);
-		bool BroadcastModelPropertyChanged(const FName PropertyName);
-		TSharedPtr<SWidget> CreateContextWidget() const;
-
-		void OnCreateBoneGroup();
-		void OnDeleteSelectedItems();
-		void OnClearBoneGroups();
-		void OnAddBoneToGroup();
+		UE_API void AddElement(TSharedPtr<FNeuralMorphBoneGroupsTreeElement> Element, TSharedPtr<FNeuralMorphBoneGroupsTreeElement> ParentElement);	
+		UE_API void HandleGetChildrenForTree(TSharedPtr<FNeuralMorphBoneGroupsTreeElement> InItem, TArray<TSharedPtr<FNeuralMorphBoneGroupsTreeElement>>& OutChildren);	
+		UE_API void OnSelectionChanged(TSharedPtr<FNeuralMorphBoneGroupsTreeElement> Selection, ESelectInfo::Type SelectInfo);
+		UE_API const TArray<TSharedPtr<FNeuralMorphBoneGroupsTreeElement>>& GetRootElements() const;
+		UE_API FReply OnKeyDown(const FGeometry& InGeometry, const FKeyEvent& InKeyEvent) override;
+		UE_API TSharedRef<ITableRow> MakeTableRowWidget(TSharedPtr<FNeuralMorphBoneGroupsTreeElement> InItem, const TSharedRef<STableViewBase>& OwnerTable);
+		UE_API TSharedPtr<SWidget> CreateContextMenuWidget() const;
+
+		UE_API void UpdateTreeElements();
+		UE_API void RefreshTree(bool bBroadcastPropertyChanged);
+		UE_API bool BroadcastModelPropertyChanged(const FName PropertyName);
+		UE_API TSharedPtr<SWidget> CreateContextWidget() const;
+
+		UE_API void OnCreateBoneGroup();
+		UE_API void OnDeleteSelectedItems();
+		UE_API void OnClearBoneGroups();
+		UE_API void OnAddBoneToGroup();
 
 	private:
 		TArray<TSharedPtr<FNeuralMorphBoneGroupsTreeElement>> RootElements;
@@ -111,3 +113,5 @@ namespace UE::NeuralMorphModel
 	};
 
 }	// namespace UE::NeuralMorphModel
+
+#undef UE_API
diff --git a/Engine/Plugins/Animation/MLDeformer/NeuralMorphModel/Source/NeuralMorphModelEditor/Public/SNeuralMorphCurveGroupsWidget.h b/Engine/Plugins/Animation/MLDeformer/NeuralMorphModel/Source/NeuralMorphModelEditor/Public/SNeuralMorphCurveGroupsWidget.h
--- a/Engine/Plugins/Animation/MLDeformer/NeuralMorphModel/Source/NeuralMorphModelEditor/Public/SNeuralMorphCurveGroupsWidget.h
+++ b/Engine/Plugins/Animation/MLDeformer/NeuralMorphModel/Source/NeuralMorphModelEditor/Public/SNeuralMorphCurveGroupsWidget.h
@@ -9,6 +9,8 @@
 #include "Framework/Commands/Commands.h"
 #include "MLDeformerEditorStyle.h"
 
+#define UE_API NEURALMORPHMODELEDITOR_API
+
 class SWidget;
 
 namespace UE::MLDeformer
@@ -68,7 +70,7 @@ namespace UE::NeuralMorphModel
 	};
 
 
-	class NEURALMORPHMODELEDITOR_API SNeuralMorphCurveGroupsWidget
+	class SNeuralMorphCurveGroupsWidget
 		: public STreeView<TSharedPtr<FNeuralMorphCurveGroupsTreeElement>>
 	{
 	public:
@@ -77,32 +79,32 @@ namespace UE::NeuralMorphModel
 		SLATE_ARGUMENT(TSharedPtr<UE::MLDeformer::SMLDeformerInputWidget>, InputWidget)
 		SLATE_END_ARGS()
 
-		void Construct(const FArguments& InArgs);
-		void BindCommands(TSharedPtr<FUICommandList> CommandList);
-		void Refresh();
+		UE_API void Construct(const FArguments& InArgs);
+		UE_API void BindCommands(TSharedPtr<FUICommandList> CommandList);
+		UE_API void Refresh();
 
-		FText GetSectionTitle() const;
-		int32 GetNumSelectedGroups() const;
-		TSharedPtr<UE::MLDeformer::SMLDeformerInputWidget> GetInputWidget() const;
+		UE_API FText GetSectionTitle() const;
+		UE_API int32 GetNumSelectedGroups() const;
+		UE_API TSharedPtr<UE::MLDeformer::SMLDeformerInputWidget> GetInputWidget() const;
 
 	private:
-		void AddElement(TSharedPtr<FNeuralMorphCurveGroupsTreeElement> Element, TSharedPtr<FNeuralMorphCurveGroupsTreeElement> ParentElement);	
-		void HandleGetChildrenForTree(TSharedPtr<FNeuralMorphCurveGroupsTreeElement> InItem, TArray<TSharedPtr<FNeuralMorphCurveGroupsTreeElement>>& OutChildren);	
-		void OnSelectionChanged(TSharedPtr<FNeuralMorphCurveGroupsTreeElement> Selection, ESelectInfo::Type SelectInfo);
-		const TArray<TSharedPtr<FNeuralMorphCurveGroupsTreeElement>>& GetRootElements() const;
-		FReply OnKeyDown(const FGeometry& InGeometry, const FKeyEvent& InKeyEvent) override;
-		TSharedRef<ITableRow> MakeTableRowWidget(TSharedPtr<FNeuralMorphCurveGroupsTreeElement> InItem, const TSharedRef<STableViewBase>& OwnerTable);
-		TSharedPtr<SWidget> CreateContextMenuWidget() const;
-
-		void UpdateTreeElements();
-		void RefreshTree(bool bBroadcastPropertyChanged);
-		bool BroadcastModelPropertyChanged(const FName PropertyName);
-		TSharedPtr<SWidget> CreateContextWidget() const;
-
-		void OnCreateCurveGroup();
-		void OnDeleteSelectedItems();
-		void OnClearCurveGroups();
-		void OnAddCurveToGroup();
+		UE_API void AddElement(TSharedPtr<FNeuralMorphCurveGroupsTreeElement> Element, TSharedPtr<FNeuralMorphCurveGroupsTreeElement> ParentElement);	
+		UE_API void HandleGetChildrenForTree(TSharedPtr<FNeuralMorphCurveGroupsTreeElement> InItem, TArray<TSharedPtr<FNeuralMorphCurveGroupsTreeElement>>& OutChildren);	
+		UE_API void OnSelectionChanged(TSharedPtr<FNeuralMorphCurveGroupsTreeElement> Selection, ESelectInfo::Type SelectInfo);
+		UE_API const TArray<TSharedPtr<FNeuralMorphCurveGroupsTreeElement>>& GetRootElements() const;
+		UE_API FReply OnKeyDown(const FGeometry& InGeometry, const FKeyEvent& InKeyEvent) override;
+		UE_API TSharedRef<ITableRow> MakeTableRowWidget(TSharedPtr<FNeuralMorphCurveGroupsTreeElement> InItem, const TSharedRef<STableViewBase>& OwnerTable);
+		UE_API TSharedPtr<SWidget> CreateContextMenuWidget() const;
+
+		UE_API void UpdateTreeElements();
+		UE_API void RefreshTree(bool bBroadcastPropertyChanged);
+		UE_API bool BroadcastModelPropertyChanged(const FName PropertyName);
+		UE_API TSharedPtr<SWidget> CreateContextWidget() const;
+
+		UE_API void OnCreateCurveGroup();
+		UE_API void OnDeleteSelectedItems();
+		UE_API void OnClearCurveGroups();
+		UE_API void OnAddCurveToGroup();
 
 	private:
 		TArray<TSharedPtr<FNeuralMorphCurveGroupsTreeElement>> RootElements;
@@ -112,3 +114,5 @@ namespace UE::NeuralMorphModel
 	};
 
 }	// namespace UE::NeuralMorphModel
+
+#undef UE_API
diff --git a/Engine/Plugins/Animation/MLDeformer/NeuralMorphModel/Source/NeuralMorphModelEditor/Public/SNeuralMorphInputWidget.h b/Engine/Plugins/Animation/MLDeformer/NeuralMorphModel/Source/NeuralMorphModelEditor/Public/SNeuralMorphInputWidget.h
--- a/Engine/Plugins/Animation/MLDeformer/NeuralMorphModel/Source/NeuralMorphModelEditor/Public/SNeuralMorphInputWidget.h
+++ b/Engine/Plugins/Animation/MLDeformer/NeuralMorphModel/Source/NeuralMorphModelEditor/Public/SNeuralMorphInputWidget.h
@@ -8,6 +8,8 @@
 #include "MeshDescription.h"
 #include "Framework/Commands/Commands.h"
 
+#define UE_API NEURALMORPHMODELEDITOR_API
+
 class FMenuBuilder;
 class USkeletalMesh;
 struct FMLDeformerMaskInfo;
@@ -45,7 +47,7 @@ namespace UE::NeuralMorphModel
 	};
 
 
-	class NEURALMORPHMODELEDITOR_API SNeuralMorphInputWidget
+	class SNeuralMorphInputWidget
 		: public UE::MLDeformer::SMLDeformerInputWidget
 	{
 	public:
@@ -53,30 +55,30 @@ namespace UE::NeuralMorphModel
 		SLATE_ARGUMENT(UE::MLDeformer::FMLDeformerEditorModel*, EditorModel)
 		SLATE_END_ARGS()
 
-		void Construct(const FArguments& InArgs);
+		UE_API void Construct(const FArguments& InArgs);
 
 		// SMLDeformerInputWidget overrides.
-		virtual void Refresh() override;
-		virtual void AddInputBonesMenuItems(FMenuBuilder& MenuBuilder) override;
-		virtual void AddInputBonesPlusIconMenuItems(FMenuBuilder& MenuBuilder) override;
-		virtual void OnAddInputBones(const TArray<FName>& Names) override;
-		virtual void OnAddInputCurves(const TArray<FName>& Names) override;
-		virtual void OnAddAnimatedBones() override;
-		virtual void OnClearInputBones() override;
-		virtual void OnDeleteInputBones(const TArray<FName>& Names) override;
-		virtual void OnDeleteInputCurves(const TArray<FName>& Names) override;
-		virtual void OnSelectInputBone(FName BoneName) override;
-		virtual void OnSelectInputCurve(FName BoneName) override;
-		virtual void OnAddAnimatedCurves() override;
-		virtual void ClearSelectionForAllWidgetsExceptThis(TSharedPtr<SWidget> ExceptThisWidget) override;
-		virtual TSharedPtr<SWidget> GetExtraBonePickerWidget() override;
+		UE_API virtual void Refresh() override;
+		UE_API virtual void AddInputBonesMenuItems(FMenuBuilder& MenuBuilder) override;
+		UE_API virtual void AddInputBonesPlusIconMenuItems(FMenuBuilder& MenuBuilder) override;
+		UE_API virtual void OnAddInputBones(const TArray<FName>& Names) override;
+		UE_API virtual void OnAddInputCurves(const TArray<FName>& Names) override;
+		UE_API virtual void OnAddAnimatedBones() override;
+		UE_API virtual void OnClearInputBones() override;
+		UE_API virtual void OnDeleteInputBones(const TArray<FName>& Names) override;
+		UE_API virtual void OnDeleteInputCurves(const TArray<FName>& Names) override;
+		UE_API virtual void OnSelectInputBone(FName BoneName) override;
+		UE_API virtual void OnSelectInputCurve(FName BoneName) override;
+		UE_API virtual void OnAddAnimatedCurves() override;
+		UE_API virtual void ClearSelectionForAllWidgetsExceptThis(TSharedPtr<SWidget> ExceptThisWidget) override;
+		UE_API virtual TSharedPtr<SWidget> GetExtraBonePickerWidget() override;
 		// ~END SMLDeformerInputWidget overrides.
 
-		void AddInputBoneGroupsMenuItems(FMenuBuilder& MenuBuilder);
-		void AddInputBoneGroupsPlusIconMenuItems(FMenuBuilder& MenuBuilder);
+		UE_API void AddInputBoneGroupsMenuItems(FMenuBuilder& MenuBuilder);
+		UE_API void AddInputBoneGroupsPlusIconMenuItems(FMenuBuilder& MenuBuilder);
 
-		void OnSelectInputBoneGroup(TSharedPtr<FNeuralMorphBoneGroupsTreeElement> Element);
-		void OnSelectInputCurveGroup(TSharedPtr<FNeuralMorphCurveGroupsTreeElement> Element);
+		UE_API void OnSelectInputBoneGroup(TSharedPtr<FNeuralMorphBoneGroupsTreeElement> Element);
+		UE_API void OnSelectInputCurveGroup(TSharedPtr<FNeuralMorphCurveGroupsTreeElement> Element);
 
 		TSharedPtr<FUICommandList> GetBoneGroupsCommandList() const	{ return BoneGroupsCommandList; }
 		TSharedPtr<FUICommandList> GetCurveGroupsCommandList() const { return CurveGroupsCommandList; }
@@ -84,23 +86,23 @@ namespace UE::NeuralMorphModel
 		int32 GetHierarchyDepth() const { return HierarchyDepth; }
 
 	protected:
-		void CreateBoneGroupsSection();
-		void CreateCurveGroupsSection();
+		UE_API void CreateBoneGroupsSection();
+		UE_API void CreateCurveGroupsSection();
 
-		void ResetAllBoneMasks();
-		void ResetSelectedBoneMasks();
-		void ExpandBoneMasks();
+		UE_API void ResetAllBoneMasks();
+		UE_API void ResetSelectedBoneMasks();
+		UE_API void ExpandBoneMasks();
 
-		void ResetAllBoneGroupMasks();
-		void ResetSelectedBoneGroupMasks();
-		void ExpandBoneGroupMasks();
+		UE_API void ResetAllBoneGroupMasks();
+		UE_API void ResetSelectedBoneGroupMasks();
+		UE_API void ExpandBoneGroupMasks();
 
-		void ConfigureBoneMask();
-		void ConfigureBoneGroupMask();
+		UE_API void ConfigureBoneMask();
+		UE_API void ConfigureBoneGroupMask();
 
-		void BindCommands();
-		FReply ShowBoneGroupsManageContextMenu();
-		FReply ShowCurveGroupsManageContextMenu();
+		UE_API void BindCommands();
+		UE_API FReply ShowBoneGroupsManageContextMenu();
+		UE_API FReply ShowCurveGroupsManageContextMenu();
 
 	protected:
 		TSharedPtr<SNeuralMorphBoneGroupsWidget> BoneGroupsWidget;
@@ -110,3 +112,5 @@ namespace UE::NeuralMorphModel
 		int32 HierarchyDepth = 2;
 	};
 }	// namespace UE::NeuralMorphModel
+
+#undef UE_API
diff --git a/Engine/Plugins/Animation/PoseSearch/Source/Editor/Private/PoseSearchDebuggerReflection.h b/Engine/Plugins/Animation/PoseSearch/Source/Editor/Private/PoseSearchDebuggerReflection.h
--- a/Engine/Plugins/Animation/PoseSearch/Source/Editor/Private/PoseSearchDebuggerReflection.h
+++ b/Engine/Plugins/Animation/PoseSearch/Source/Editor/Private/PoseSearchDebuggerReflection.h
@@ -8,8 +8,8 @@
 /**
  * Reflection UObject being observed in the details view panel of the debugger
  */
-UCLASS()
-class POSESEARCHEDITOR_API UPoseSearchDebuggerReflection : public UObject
+UCLASS(MinimalAPI)
+class UPoseSearchDebuggerReflection : public UObject
 {
 	GENERATED_BODY()
 
diff --git a/Engine/Plugins/Animation/PoseSearch/Source/Editor/Public/PoseSearchDatabaseFactory.h b/Engine/Plugins/Animation/PoseSearch/Source/Editor/Public/PoseSearchDatabaseFactory.h
--- a/Engine/Plugins/Animation/PoseSearch/Source/Editor/Public/PoseSearchDatabaseFactory.h
+++ b/Engine/Plugins/Animation/PoseSearch/Source/Editor/Public/PoseSearchDatabaseFactory.h
@@ -7,11 +7,13 @@
 #include "Factories/Factory.h"
 #include "PoseSearchDatabaseFactory.generated.h"
 
+#define UE_API POSESEARCHEDITOR_API
+
 struct FAssetData;
 class SWindow;
 
-UCLASS()
-class POSESEARCHEDITOR_API UPoseSearchDatabaseFactory : public UFactory
+UCLASS(MinimalAPI)
+class UPoseSearchDatabaseFactory : public UFactory
 {
 	GENERATED_UCLASS_BODY()
 
@@ -21,10 +23,12 @@ class POSESEARCHEDITOR_API UPoseSearchDatabaseFactory : public UFactory
 	TSharedPtr<SWindow> PickerWindow;
 
 	// UFactory interface
-	virtual bool ConfigureProperties() override;
-	virtual UObject* FactoryCreateNew(UClass* Class, UObject* InParent, FName Name, EObjectFlags Flags, UObject* Context, FFeedbackContext* Warn) override;
-	virtual FString GetDefaultNewAssetName() const override;
+	UE_API virtual bool ConfigureProperties() override;
+	UE_API virtual UObject* FactoryCreateNew(UClass* Class, UObject* InParent, FName Name, EObjectFlags Flags, UObject* Context, FFeedbackContext* Warn) override;
+	UE_API virtual FString GetDefaultNewAssetName() const override;
 	// End of UFactory interface
 
-	void OnTargetSchemaSelected(const FAssetData& SelectedAsset);
+	UE_API void OnTargetSchemaSelected(const FAssetData& SelectedAsset);
 };
+
+#undef UE_API
diff --git a/Engine/Plugins/Animation/PoseSearch/Source/Editor/Public/PoseSearchNormalizationSetFactory.h b/Engine/Plugins/Animation/PoseSearch/Source/Editor/Public/PoseSearchNormalizationSetFactory.h
--- a/Engine/Plugins/Animation/PoseSearch/Source/Editor/Public/PoseSearchNormalizationSetFactory.h
+++ b/Engine/Plugins/Animation/PoseSearch/Source/Editor/Public/PoseSearchNormalizationSetFactory.h
@@ -7,13 +7,17 @@
 #include "Factories/Factory.h"
 #include "PoseSearchNormalizationSetFactory.generated.h"
 
-UCLASS()
-class POSESEARCHEDITOR_API UPoseSearchNormalizationSetFactory : public UFactory
+#define UE_API POSESEARCHEDITOR_API
+
+UCLASS(MinimalAPI)
+class UPoseSearchNormalizationSetFactory : public UFactory
 {
 	GENERATED_UCLASS_BODY()
 
 	// UFactory interface
-	virtual UObject* FactoryCreateNew(UClass* Class, UObject* InParent, FName Name, EObjectFlags Flags, UObject* Context, FFeedbackContext* Warn) override;
-	virtual FString GetDefaultNewAssetName() const override;
+	UE_API virtual UObject* FactoryCreateNew(UClass* Class, UObject* InParent, FName Name, EObjectFlags Flags, UObject* Context, FFeedbackContext* Warn) override;
+	UE_API virtual FString GetDefaultNewAssetName() const override;
 	// End of UFactory interface
 };
+
+#undef UE_API
diff --git a/Engine/Plugins/Animation/PoseSearch/Source/Editor/Public/PoseSearchSchemaFactory.h b/Engine/Plugins/Animation/PoseSearch/Source/Editor/Public/PoseSearchSchemaFactory.h
--- a/Engine/Plugins/Animation/PoseSearch/Source/Editor/Public/PoseSearchSchemaFactory.h
+++ b/Engine/Plugins/Animation/PoseSearch/Source/Editor/Public/PoseSearchSchemaFactory.h
@@ -7,11 +7,13 @@
 #include "Factories/Factory.h"
 #include "PoseSearchSchemaFactory.generated.h"
 
+#define UE_API POSESEARCHEDITOR_API
+
 struct FAssetData;
 class SWindow;
 
-UCLASS()
-class POSESEARCHEDITOR_API UPoseSearchSchemaFactory : public UFactory
+UCLASS(MinimalAPI)
+class UPoseSearchSchemaFactory : public UFactory
 {
 	GENERATED_UCLASS_BODY()
 
@@ -21,10 +23,12 @@ class POSESEARCHEDITOR_API UPoseSearchSchemaFactory : public UFactory
 	TSharedPtr<SWindow> PickerWindow;
 
 	// UFactory interface
-	virtual bool ConfigureProperties() override;
-	virtual UObject* FactoryCreateNew(UClass* Class, UObject* InParent, FName Name, EObjectFlags Flags, UObject* Context, FFeedbackContext* Warn) override;
-	virtual FString GetDefaultNewAssetName() const override;
+	UE_API virtual bool ConfigureProperties() override;
+	UE_API virtual UObject* FactoryCreateNew(UClass* Class, UObject* InParent, FName Name, EObjectFlags Flags, UObject* Context, FFeedbackContext* Warn) override;
+	UE_API virtual FString GetDefaultNewAssetName() const override;
 	// End of UFactory interface
 
-	void OnTargetSkeletonSelected(const FAssetData& SelectedAsset);
+	UE_API void OnTargetSkeletonSelected(const FAssetData& SelectedAsset);
 };
+
+#undef UE_API
diff --git a/Engine/Plugins/Animation/RigLogic/Source/RigLogicDeveloper/Public/AnimGraphNode_RigLogic.h b/Engine/Plugins/Animation/RigLogic/Source/RigLogicDeveloper/Public/AnimGraphNode_RigLogic.h
--- a/Engine/Plugins/Animation/RigLogic/Source/RigLogicDeveloper/Public/AnimGraphNode_RigLogic.h
+++ b/Engine/Plugins/Animation/RigLogic/Source/RigLogicDeveloper/Public/AnimGraphNode_RigLogic.h
@@ -6,8 +6,10 @@
 #include "AnimNode_RigLogic.h"
 #include "AnimGraphNode_RigLogic.generated.h"
 
-UCLASS(meta = (Keywords = "Rig Logic Animation Node"))
-class RIGLOGICDEVELOPER_API UAnimGraphNode_RigLogic : public UAnimGraphNode_Base
+#define UE_API RIGLOGICDEVELOPER_API
+
+UCLASS(MinimalAPI, meta = (Keywords = "Rig Logic Animation Node"))
+class UAnimGraphNode_RigLogic : public UAnimGraphNode_Base
 {
 	GENERATED_UCLASS_BODY()
 
@@ -15,9 +17,11 @@ class RIGLOGICDEVELOPER_API UAnimGraphNode_RigLogic : public UAnimGraphNode_Base
 	FAnimNode_RigLogic Node;
 
 public:
-	FText GetTooltipText() const;
-	FText GetNodeTitle(ENodeTitleType::Type TitleType) const;
+	UE_API FText GetTooltipText() const;
+	UE_API FText GetNodeTitle(ENodeTitleType::Type TitleType) const;
 
 protected:
-	void ValidateAnimNodeDuringCompilation(USkeleton* ForSkeleton, FCompilerResultsLog& MessageLog);
+	UE_API void ValidateAnimNodeDuringCompilation(USkeleton* ForSkeleton, FCompilerResultsLog& MessageLog);
 };
+
+#undef UE_API
diff --git a/Engine/Plugins/Animation/RigLogic/Source/RigLogicDeveloper/Public/RigLogicDeveloper.h b/Engine/Plugins/Animation/RigLogic/Source/RigLogicDeveloper/Public/RigLogicDeveloper.h
--- a/Engine/Plugins/Animation/RigLogic/Source/RigLogicDeveloper/Public/RigLogicDeveloper.h
+++ b/Engine/Plugins/Animation/RigLogic/Source/RigLogicDeveloper/Public/RigLogicDeveloper.h
@@ -5,11 +5,15 @@
 #include "CoreMinimal.h"
 #include "Modules/ModuleInterface.h"
 
+#define UE_API RIGLOGICDEVELOPER_API
+
 DECLARE_LOG_CATEGORY_EXTERN(LogRigLogicDeveloper, Log, All);
 
-class RIGLOGICDEVELOPER_API FRigLogicDeveloperModule : public IModuleInterface 
+class FRigLogicDeveloperModule : public IModuleInterface 
 {
 public:
-	void StartupModule() override;
-	void ShutdownModule() override;
+	UE_API void StartupModule() override;
+	UE_API void ShutdownModule() override;
 };
+
+#undef UE_API
diff --git a/Engine/Plugins/Animation/RigLogic/Source/RigLogicEditor/Public/DNAImporterLibrary.h b/Engine/Plugins/Animation/RigLogic/Source/RigLogicEditor/Public/DNAImporterLibrary.h
--- a/Engine/Plugins/Animation/RigLogic/Source/RigLogicEditor/Public/DNAImporterLibrary.h
+++ b/Engine/Plugins/Animation/RigLogic/Source/RigLogicEditor/Public/DNAImporterLibrary.h
@@ -6,15 +6,19 @@
 #include "Kismet/BlueprintFunctionLibrary.h"
 #include "DNAImporterLibrary.generated.h"
 
+#define UE_API RIGLOGICEDITOR_API
+
 /**
  * 
  */
-UCLASS()
-class RIGLOGICEDITOR_API UDNAImporterLibrary : public UBlueprintFunctionLibrary
+UCLASS(MinimalAPI)
+class UDNAImporterLibrary : public UBlueprintFunctionLibrary
 {
 	GENERATED_BODY()
 	
 	/** Import a DNA file onto the specified mesh asset.*/
 	UFUNCTION(CallInEditor, BlueprintCallable, meta = (DisplayName = "Import Skeletal Mesh DNA File", Keywords = "DNA Importer Skeletal Mesh"), Category = "DNA")
-	static void ImportSkeletalMeshDNA(const FString FileName, UObject* Mesh);
+	static UE_API void ImportSkeletalMeshDNA(const FString FileName, UObject* Mesh);
 };
+
+#undef UE_API
diff --git a/Engine/Plugins/Animation/RigLogic/Source/RigLogicEditor/Public/RigLogicEditor.h b/Engine/Plugins/Animation/RigLogic/Source/RigLogicEditor/Public/RigLogicEditor.h
--- a/Engine/Plugins/Animation/RigLogic/Source/RigLogicEditor/Public/RigLogicEditor.h
+++ b/Engine/Plugins/Animation/RigLogic/Source/RigLogicEditor/Public/RigLogicEditor.h
@@ -6,25 +6,29 @@
 #include "UObject/Object.h"
 #include "Modules/ModuleInterface.h"
 
+#define UE_API RIGLOGICEDITOR_API
+
 class FExtender;
 class FMenuBuilder;
 struct FAssetData;
 
-class RIGLOGICEDITOR_API FRigLogicEditor : public IModuleInterface 
+class FRigLogicEditor : public IModuleInterface 
 {
 public:
-	void StartupModule() override;
-	void ShutdownModule() override;
+	UE_API void StartupModule() override;
+	UE_API void ShutdownModule() override;
 
 private:
 	TArray<TSharedRef<class IAssetTypeActions>> AssetTypeActions;
 
-	static TSharedRef<FExtender> OnExtendSkelMeshWithDNASelectionMenu(const TArray<FAssetData>& SelectedAssets);
-	static void CreateDnaActionsSubMenu(FMenuBuilder& MenuBuilder, const TArray<FAssetData> SelectedAssets);
-	static void GetDNAMenu(FMenuBuilder& MenuBuilder, const TArray<FAssetData> SelectedAssets);
+	static UE_API TSharedRef<FExtender> OnExtendSkelMeshWithDNASelectionMenu(const TArray<FAssetData>& SelectedAssets);
+	static UE_API void CreateDnaActionsSubMenu(FMenuBuilder& MenuBuilder, const TArray<FAssetData> SelectedAssets);
+	static UE_API void GetDNAMenu(FMenuBuilder& MenuBuilder, const TArray<FAssetData> SelectedAssets);
 
-	static void ExecuteDNAImport(UObject* Mesh);
-	static void ExecuteDNAReimport(UObject* Mesh);
+	static UE_API void ExecuteDNAImport(UObject* Mesh);
+	static UE_API void ExecuteDNAReimport(UObject* Mesh);
 
-	static void GetAssetRegistryTagsForDNA(FAssetRegistryTagsContext Context);
+	static UE_API void GetAssetRegistryTagsForDNA(FAssetRegistryTagsContext Context);
 };
+
+#undef UE_API
diff --git a/Engine/Plugins/Animation/SkeletalMeshModelingTools/Source/SkeletalMeshModelingTools/Private/SkeletalMeshEditorUtils.h b/Engine/Plugins/Animation/SkeletalMeshModelingTools/Source/SkeletalMeshModelingTools/Private/SkeletalMeshEditorUtils.h
--- a/Engine/Plugins/Animation/SkeletalMeshModelingTools/Source/SkeletalMeshModelingTools/Private/SkeletalMeshEditorUtils.h
+++ b/Engine/Plugins/Animation/SkeletalMeshModelingTools/Source/SkeletalMeshModelingTools/Private/SkeletalMeshEditorUtils.h
@@ -10,6 +10,8 @@
 
 #include "SkeletalMeshEditorUtils.generated.h"
 
+#define UE_API SKELETALMESHMODELINGTOOLS_API
+
 class UInteractiveToolsContext;
 class UInteractiveToolManager;
 class ISkeletalMeshEditor;
@@ -33,22 +35,22 @@ namespace SkeletalMeshEditorUtils
  * USkeletalMeshEditorContextObject
  */
 
-UCLASS()
-class SKELETALMESHMODELINGTOOLS_API USkeletalMeshEditorContextObject : public USkeletalMeshEditorContextObjectBase
+UCLASS(MinimalAPI)
+class USkeletalMeshEditorContextObject : public USkeletalMeshEditorContextObjectBase
 {
 	GENERATED_BODY()
 
 public:
-	void Register(UInteractiveToolManager* InToolManager);
-	void Unregister(UInteractiveToolManager* InToolManager);
+	UE_API void Register(UInteractiveToolManager* InToolManager);
+	UE_API void Unregister(UInteractiveToolManager* InToolManager);
 
-	void Init(const TWeakPtr<ISkeletalMeshEditor>& InEditor);
+	UE_API void Init(const TWeakPtr<ISkeletalMeshEditor>& InEditor);
 
-	virtual void HideSkeleton() override;
-	virtual void ShowSkeleton() override;
+	UE_API virtual void HideSkeleton() override;
+	UE_API virtual void ShowSkeleton() override;
 
-	virtual void BindTo(ISkeletalMeshEditingInterface* InEditingInterface) override;
-	virtual void UnbindFrom(ISkeletalMeshEditingInterface* InEditingInterface) override;
+	UE_API virtual void BindTo(ISkeletalMeshEditingInterface* InEditingInterface) override;
+	UE_API virtual void UnbindFrom(ISkeletalMeshEditingInterface* InEditingInterface) override;
 	
 private:
 	TWeakPtr<ISkeletalMeshEditor> Editor = nullptr;
@@ -56,7 +58,7 @@ private:
 	UPROPERTY()
 	ESkeletonDrawMode SkeletonDrawMode = ESkeletonDrawMode::Default;
 
-	TSharedPtr<ISkeletalMeshEditorBinding> GetBinding() const;
+	UE_API TSharedPtr<ISkeletalMeshEditorBinding> GetBinding() const;
 
 	struct FBindData
 	{
@@ -66,22 +68,22 @@ private:
 	TMap< ISkeletalMeshEditingInterface*, FBindData > EditorBindings;
 	TMap< ISkeletalMeshEditingInterface*, FBindData > TreeBindings;
 
-	static TPair<FDelegateHandle, FDelegateHandle> BindInterfaceTo(
+	static UE_API TPair<FDelegateHandle, FDelegateHandle> BindInterfaceTo(
 		ISkeletalMeshEditingInterface* InInterface,
 		ISkeletalMeshNotifier& InOtherNotifier);
 	
-	static void UnbindInterfaceFrom(
+	static UE_API void UnbindInterfaceFrom(
 		ISkeletalMeshEditingInterface* InInterface,
 		ISkeletalMeshNotifier& InOtherNotifier,
 		const FBindData& InOutBindData);
 
-	void BindEditor(ISkeletalMeshEditingInterface* InEditingInterface);
-	void UnbindEditor(ISkeletalMeshEditingInterface* InEditingInterface);
+	UE_API void BindEditor(ISkeletalMeshEditingInterface* InEditingInterface);
+	UE_API void UnbindEditor(ISkeletalMeshEditingInterface* InEditingInterface);
 
-	void BindRefSkeletonTree(ISkeletalMeshEditingInterface* InEditingInterface);
-	void UnbindRefSkeletonTree(ISkeletalMeshEditingInterface* InEditingInterface);
+	UE_API void BindRefSkeletonTree(ISkeletalMeshEditingInterface* InEditingInterface);
+	UE_API void UnbindRefSkeletonTree(ISkeletalMeshEditingInterface* InEditingInterface);
 
-	static const FName& GetSkeletonTreeTabId();
+	static UE_API const FName& GetSkeletonTreeTabId();
 
 	TSharedPtr<SWidget> RefSkeletonWidget;
 	TSharedPtr<SReferenceSkeletonTree> RefSkeletonTree;
@@ -89,3 +91,5 @@ private:
 	
 	bool bRegistered = false;
 };
+
+#undef UE_API
diff --git a/Engine/Plugins/Animation/SkeletalMeshModelingTools/Source/SkeletalMeshModelingTools/Private/SkeletalMeshGizmoUtils.h b/Engine/Plugins/Animation/SkeletalMeshModelingTools/Source/SkeletalMeshModelingTools/Private/SkeletalMeshGizmoUtils.h
--- a/Engine/Plugins/Animation/SkeletalMeshModelingTools/Source/SkeletalMeshModelingTools/Private/SkeletalMeshGizmoUtils.h
+++ b/Engine/Plugins/Animation/SkeletalMeshModelingTools/Source/SkeletalMeshModelingTools/Private/SkeletalMeshGizmoUtils.h
@@ -8,6 +8,8 @@
 
 #include "SkeletalMeshGizmoUtils.generated.h"
 
+#define UE_API SKELETALMESHMODELINGTOOLS_API
+
 class UInteractiveToolsContext;
 class UInteractiveToolManager;
 class UTransformGizmo;
@@ -50,17 +52,17 @@ namespace SkeletalMeshGizmoUtils
  * be used to update skeletal mesh infos using a Gizmo.
  */
 
-UCLASS()
-class SKELETALMESHMODELINGTOOLS_API USkeletalMeshGizmoWrapper : public USkeletalMeshGizmoWrapperBase
+UCLASS(MinimalAPI)
+class USkeletalMeshGizmoWrapper : public USkeletalMeshGizmoWrapperBase
 {
 	GENERATED_BODY()
 public:
-	virtual void Initialize(const FTransform& InTransform = FTransform::Identity, const EToolContextCoordinateSystem& InTransformMode = EToolContextCoordinateSystem::Local) override;
-	virtual void HandleBoneTransform(FGetTransform GetTransformFunc, FSetTransform SetTransformFunc) override;
-	virtual void Clear() override;
+	UE_API virtual void Initialize(const FTransform& InTransform = FTransform::Identity, const EToolContextCoordinateSystem& InTransformMode = EToolContextCoordinateSystem::Local) override;
+	UE_API virtual void HandleBoneTransform(FGetTransform GetTransformFunc, FSetTransform SetTransformFunc) override;
+	UE_API virtual void Clear() override;
 
-	virtual bool CanInteract() const override;
-	virtual bool IsGizmoHit(const FInputDeviceRay& PressPos) const override;
+	UE_API virtual bool CanInteract() const override;
+	UE_API virtual bool IsGizmoHit(const FInputDeviceRay& PressPos) const override;
 
 	UPROPERTY()
 	TObjectPtr<UTransformGizmo> TransformGizmo;
@@ -73,24 +75,26 @@ public:
  * (see UCombinedTransformGizmoContextObject for more details)
  */
 
-UCLASS()
-class SKELETALMESHMODELINGTOOLS_API USkeletalMeshGizmoContextObject : public USkeletalMeshGizmoContextObjectBase
+UCLASS(MinimalAPI)
+class USkeletalMeshGizmoContextObject : public USkeletalMeshGizmoContextObjectBase
 {
 	GENERATED_BODY()
 
 public:
 
 	// builder identifiers for transform gizmo
-	static const FString& TransformBuilderIdentifier();
+	static UE_API const FString& TransformBuilderIdentifier();
 	
-	void RegisterGizmosWithManager(UInteractiveToolManager* InToolManager);
-	void UnregisterGizmosWithManager(UInteractiveToolManager* InToolManager);
+	UE_API void RegisterGizmosWithManager(UInteractiveToolManager* InToolManager);
+	UE_API void UnregisterGizmosWithManager(UInteractiveToolManager* InToolManager);
 
-	virtual USkeletalMeshGizmoWrapperBase* GetNewWrapper(UInteractiveToolManager* InToolManager, UObject* Outer = nullptr, IGizmoStateTarget* InStateTarget = nullptr) override;
+	UE_API virtual USkeletalMeshGizmoWrapperBase* GetNewWrapper(UInteractiveToolManager* InToolManager, UObject* Outer = nullptr, IGizmoStateTarget* InStateTarget = nullptr) override;
 
 private:
 	
-	static const FGizmoCustomization& GetGizmoCustomization();
+	static UE_API const FGizmoCustomization& GetGizmoCustomization();
 	
 	bool bGizmosRegistered = false;
 };
+
+#undef UE_API
diff --git a/Engine/Plugins/Animation/TweeningUtils/Source/TweeningUtilsEditor/Public/Math/Abstraction/TweenModelArray.h b/Engine/Plugins/Animation/TweeningUtils/Source/TweeningUtilsEditor/Public/Math/Abstraction/TweenModelArray.h
--- a/Engine/Plugins/Animation/TweeningUtils/Source/TweeningUtilsEditor/Public/Math/Abstraction/TweenModelArray.h
+++ b/Engine/Plugins/Animation/TweeningUtils/Source/TweeningUtilsEditor/Public/Math/Abstraction/TweenModelArray.h
@@ -1,4 +1,4 @@
-﻿// Copyright Epic Games, Inc. All Rights Reserved.
+// Copyright Epic Games, Inc. All Rights Reserved.
 
 #pragma once
 
@@ -7,25 +7,27 @@
 #include "TweenModelDisplayInfo.h"
 #include "Templates/UnrealTemplate.h"
 
+#define UE_API TWEENINGUTILSEDITOR_API
+
 namespace UE::TweeningUtilsEditor
 {
 /** Backed by an array of tween models and their UI info. */
-class TWEENINGUTILSEDITOR_API FTweenModelArray : public ITweenModelContainer, public FNoncopyable
+class FTweenModelArray : public ITweenModelContainer, public FNoncopyable
 {
 public:
 	
 	explicit FTweenModelArray(TArray<FTweenModelUIEntry> InTweenModels) : TweenModels(MoveTemp(InTweenModels)) {}
 
 	//~ Begin FTweenModelContainer Interface
-	virtual void ForEachModel(const TFunctionRef<void(FTweenModel&)> InConsumer) const override;
-	virtual FTweenModel* GetModel(int32 InIndex) const override;
+	UE_API virtual void ForEachModel(const TFunctionRef<void(FTweenModel&)> InConsumer) const override;
+	UE_API virtual FTweenModel* GetModel(int32 InIndex) const override;
 	virtual int32 NumModels() const override { return TweenModels.Num(); }
-	virtual TSharedPtr<FUICommandInfo> GetCommandForModel(const FTweenModel& InTweenModel) const override;
-	virtual const FSlateBrush* GetIconForModel(const FTweenModel& InTweenModel) const override;
-	virtual FLinearColor GetColorForModel(const FTweenModel& InTweenModel) const override;
-	virtual FText GetLabelForModel(const FTweenModel& InTweenModel) const override;
-	virtual FText GetToolTipForModel(const FTweenModel& InTweenModel) const override;
-	virtual FString GetModelIdentifier(const FTweenModel& InTweenModel) const override;
+	UE_API virtual TSharedPtr<FUICommandInfo> GetCommandForModel(const FTweenModel& InTweenModel) const override;
+	UE_API virtual const FSlateBrush* GetIconForModel(const FTweenModel& InTweenModel) const override;
+	UE_API virtual FLinearColor GetColorForModel(const FTweenModel& InTweenModel) const override;
+	UE_API virtual FText GetLabelForModel(const FTweenModel& InTweenModel) const override;
+	UE_API virtual FText GetToolTipForModel(const FTweenModel& InTweenModel) const override;
+	UE_API virtual FString GetModelIdentifier(const FTweenModel& InTweenModel) const override;
 	//~ Begin FTweenModelContainer Interface
 
 private:
@@ -34,3 +36,5 @@ private:
 	const TArray<FTweenModelUIEntry> TweenModels;
 };
 }
+
+#undef UE_API
diff --git a/Engine/Plugins/Animation/TweeningUtils/Source/TweeningUtilsEditor/Public/Math/ContiguousKeyMapping.h b/Engine/Plugins/Animation/TweeningUtils/Source/TweeningUtilsEditor/Public/Math/ContiguousKeyMapping.h
--- a/Engine/Plugins/Animation/TweeningUtils/Source/TweeningUtilsEditor/Public/Math/ContiguousKeyMapping.h
+++ b/Engine/Plugins/Animation/TweeningUtils/Source/TweeningUtilsEditor/Public/Math/ContiguousKeyMapping.h
@@ -1,4 +1,4 @@
-﻿// Copyright Epic Games, Inc. All Rights Reserved.
+// Copyright Epic Games, Inc. All Rights Reserved.
 
 #pragma once
 
@@ -9,12 +9,14 @@
 #include "HAL/Platform.h"
 #include "Misc/CoreMiscDefines.h"
 
+#define UE_API TWEENINGUTILSEDITOR_API
+
 class FCurveEditor;
 
 namespace UE::TweeningUtilsEditor
 {
 /** Unbroken chain of keys. */
-struct TWEENINGUTILSEDITOR_API FContiguousKeys
+struct FContiguousKeys
 {
 	/** Indices of keys to blend */
 	const TArray<int32> Indices;
@@ -34,7 +36,7 @@ struct TWEENINGUTILSEDITOR_API FContiguousKeys
  * Stores all key positions as they were before the blend, and all the chains of keys that are supposed to be blended.
  * You can use this type independently from FContiguousKeyMapping. For example, you can use this to call TweeningUtilsEditor::TweenRange directly.
  */
-struct TWEENINGUTILSEDITOR_API FBlendRangesData
+struct FBlendRangesData
 {
 	const TArray<FVector2d> AllKeyPositions;
 
@@ -74,28 +76,28 @@ struct TWEENINGUTILSEDITOR_API FBlendRangesData
 	/** @return The key after the blend range defined by InKeys */
 	const FVector2d& GetAfterBlendRange(const FContiguousKeys& InKeys) const;
 
-	void AddBlendRange(const TArray<int32>& InContiguousKeyIndices);
+	UE_API void AddBlendRange(const TArray<int32>& InContiguousKeyIndices);
 };
 	
 /** Information for blending a map of curves and the keys in it to blend. Saves the old key values for the duration of the blend.*/
-struct TWEENINGUTILSEDITOR_API FContiguousKeyMapping
+struct FContiguousKeyMapping
 {
 	/** Holds the corresponding key handles for all the keys. */
-	struct TWEENINGUTILSEDITOR_API FContiguousKeysArray : FBlendRangesData
+	struct FContiguousKeysArray : FBlendRangesData
 	{
 		const TArray<FKeyHandle> AllKeyHandles;
 		
 		explicit FContiguousKeysArray(const TArray<FKeyHandle>& InAllKeyHandles, const TArray<FVector2d>& InAllKeyPositions)
 			: FBlendRangesData(InAllKeyPositions) , AllKeyHandles(InAllKeyHandles) {}
 		
-		void AddBlendRange(const TArray<int32>& ContiguousKeyIndices);
+		UE_API void AddBlendRange(const TArray<int32>& ContiguousKeyIndices);
 	};
 
 	FContiguousKeyMapping() = default;
-	explicit FContiguousKeyMapping(const FCurveEditor& InCurveEditor);
+	UE_API explicit FContiguousKeyMapping(const FCurveEditor& InCurveEditor);
 
 	/** Appends keys from the given curve.*/
-	void Append(const FCurveEditor& InCurveEditor, const FCurveModelID& InCurveId, TConstArrayView<FKeyHandle> InKeysToBlend);
+	UE_API void Append(const FCurveEditor& InCurveEditor, const FCurveModelID& InCurveId, TConstArrayView<FKeyHandle> InKeysToBlend);
 
 	TMap<FCurveModelID, FContiguousKeysArray> KeyMap;
 };
@@ -139,4 +141,6 @@ inline const FVector2d& FBlendRangesData::GetAfterBlendRange(const FContiguousKe
 	const int32 AfterIndex = InKeys.NextIndex != INDEX_NONE ? InKeys.NextIndex : InKeys.Indices[InKeys.Indices.Num() - 1];
 	return AllKeyPositions[AfterIndex];
 }
-}
\ No newline at end of file
+}
+
+#undef UE_API
diff --git a/Engine/Plugins/Animation/TweeningUtils/Source/TweeningUtilsEditor/Public/Math/Models/CurveTimeOffsetTweenModel.h b/Engine/Plugins/Animation/TweeningUtils/Source/TweeningUtilsEditor/Public/Math/Models/CurveTimeOffsetTweenModel.h
--- a/Engine/Plugins/Animation/TweeningUtils/Source/TweeningUtilsEditor/Public/Math/Models/CurveTimeOffsetTweenModel.h
+++ b/Engine/Plugins/Animation/TweeningUtils/Source/TweeningUtilsEditor/Public/Math/Models/CurveTimeOffsetTweenModel.h
@@ -9,21 +9,23 @@
 #include "Templates/SharedPointer.h"
 #include "Templates/UnrealTemplate.h"
 
+#define UE_API TWEENINGUTILSEDITOR_API
+
 class FCurveEditor;
 
 namespace UE::TweeningUtilsEditor
 {
 /** Recomputes all keys output values such that the curve is effectively shifted to left and right without modifying the keys' input values. */
-class TWEENINGUTILSEDITOR_API FCurveTimeOffsetTweenModel : public FTweenModel, public FNoncopyable
+class FCurveTimeOffsetTweenModel : public FTweenModel, public FNoncopyable
 {
 public:
 
-	explicit FCurveTimeOffsetTweenModel(TAttribute<TWeakPtr<FCurveEditor>> InWeakCurveEditor);
+	UE_API explicit FCurveTimeOffsetTweenModel(TAttribute<TWeakPtr<FCurveEditor>> InWeakCurveEditor);
 	
 	//~ Begin FTweenModel Interface
-	virtual void StartBlendOperation() override;
-	virtual void StopBlendOperation() override;
-	virtual void BlendValues(float InNormalizedValue) override;
+	UE_API virtual void StartBlendOperation() override;
+	UE_API virtual void StopBlendOperation() override;
+	UE_API virtual void BlendValues(float InNormalizedValue) override;
 	//~ Begin FTweenModel Interface
 
 protected:
@@ -39,3 +41,4 @@ protected:
 };
 }
 
+#undef UE_API
diff --git a/Engine/Plugins/Animation/TweeningUtils/Source/TweeningUtilsEditor/Public/Math/Models/TweenModel.h b/Engine/Plugins/Animation/TweeningUtils/Source/TweeningUtilsEditor/Public/Math/Models/TweenModel.h
--- a/Engine/Plugins/Animation/TweeningUtils/Source/TweeningUtilsEditor/Public/Math/Models/TweenModel.h
+++ b/Engine/Plugins/Animation/TweeningUtils/Source/TweeningUtilsEditor/Public/Math/Models/TweenModel.h
@@ -10,7 +10,7 @@ namespace UE::TweeningUtilsEditor
  * Implements the logic of tweening some values: it could be curves, it could be control rig models.
  * Acts as model in a Model-View-Controller architecture.
  */
-class TWEENINGUTILSEDITOR_API FTweenModel
+class FTweenModel
 {
 public:
 
diff --git a/Engine/Plugins/Animation/TweeningUtils/Source/TweeningUtilsEditor/Public/TweeningToolsUserSettings.h b/Engine/Plugins/Animation/TweeningUtils/Source/TweeningUtilsEditor/Public/TweeningToolsUserSettings.h
--- a/Engine/Plugins/Animation/TweeningUtils/Source/TweeningUtilsEditor/Public/TweeningToolsUserSettings.h
+++ b/Engine/Plugins/Animation/TweeningUtils/Source/TweeningUtilsEditor/Public/TweeningToolsUserSettings.h
@@ -8,13 +8,15 @@
 #include "UObject/NameTypes.h"
 #include "TweeningToolsUserSettings.generated.h"
 
-UCLASS(Config=EditorPerProjectUserSettings)
-class TWEENINGUTILSEDITOR_API UTweeningToolsUserSettings : public UObject
+#define UE_API TWEENINGUTILSEDITOR_API
+
+UCLASS(MinimalAPI, Config=EditorPerProjectUserSettings)
+class UTweeningToolsUserSettings : public UObject
 {
 	GENERATED_BODY()
 public:
 
-	static UTweeningToolsUserSettings* Get();
+	static UE_API UTweeningToolsUserSettings* Get();
 
 	/** Sets the preferred tween function for the given feature. */
 	void SetPreferredTweenFunction(FName InFeatureKey, FString InNewPreferredFunction)
@@ -35,3 +37,5 @@ private:
 	UPROPERTY(Config)
 	TMap<FName, FString> PreferredTweenFunction;
 };
+
+#undef UE_API
diff --git a/Engine/Plugins/Animation/TweeningUtils/Source/TweeningUtilsEditor/Public/Widgets/MVC/CycleFunctionController.h b/Engine/Plugins/Animation/TweeningUtils/Source/TweeningUtilsEditor/Public/Widgets/MVC/CycleFunctionController.h
--- a/Engine/Plugins/Animation/TweeningUtils/Source/TweeningUtilsEditor/Public/Widgets/MVC/CycleFunctionController.h
+++ b/Engine/Plugins/Animation/TweeningUtils/Source/TweeningUtilsEditor/Public/Widgets/MVC/CycleFunctionController.h
@@ -7,6 +7,8 @@
 #include "Templates/SharedPointer.h"
 #include "Templates/UnrealTemplate.h"
 
+#define UE_API TWEENINGUTILSEDITOR_API
+
 class FUICommandInfo;
 class FUICommandList;
 
@@ -16,7 +18,7 @@ class FTweenModel;
 class ITweenModelContainer;
 
 /** Cycles through functions on a tween model. */
-class TWEENINGUTILSEDITOR_API FCycleFunctionController : public FNoncopyable
+class FCycleFunctionController : public FNoncopyable
 {
 public:
 
@@ -29,7 +31,7 @@ public:
 	 * @param InCommandList The command list to add / remove the command to / from.
 	 * @param InCycleCommand The command that cycles the function
 	 */
-	explicit FCycleFunctionController(
+	UE_API explicit FCycleFunctionController(
 		TAttribute<const FTweenModel*> InCurrentTweenModelAttr,
 		const TSharedRef<ITweenModelContainer>& InTweenModelContainer,
 		FHandleTweenChange InHandleTweenChange,
@@ -37,13 +39,13 @@ public:
 		TSharedPtr<FUICommandInfo> InCycleCommand
 		);
 	/** Version that defauls the cycle command to FTweenUtilsCommands::ChangeAnimSliderTool. */
-	explicit FCycleFunctionController(
+	UE_API explicit FCycleFunctionController(
 		TAttribute<const FTweenModel*> InCurrentTweenModelAttr,
 		const TSharedRef<ITweenModelContainer>& InTweenModelContainer,
 		FHandleTweenChange InHandleTweenChange,
 		const TSharedRef<FUICommandList>& InCommandList
 		); 
-	virtual ~FCycleFunctionController();
+	UE_API virtual ~FCycleFunctionController();
 	
 private:
 	
@@ -60,6 +62,8 @@ private:
 	const TSharedPtr<FUICommandInfo> CycleCommand;
 	
 	/** Sets the next tween function. */
-	void CycleToNextFunction();
+	UE_API void CycleToNextFunction();
 };
 }
+
+#undef UE_API
diff --git a/Engine/Plugins/Animation/TweeningUtils/Source/TweeningUtilsEditor/Public/Widgets/MVC/MouseSlidingController.h b/Engine/Plugins/Animation/TweeningUtils/Source/TweeningUtilsEditor/Public/Widgets/MVC/MouseSlidingController.h
--- a/Engine/Plugins/Animation/TweeningUtils/Source/TweeningUtilsEditor/Public/Widgets/MVC/MouseSlidingController.h
+++ b/Engine/Plugins/Animation/TweeningUtils/Source/TweeningUtilsEditor/Public/Widgets/MVC/MouseSlidingController.h
@@ -9,6 +9,8 @@
 #include "Templates/SharedPointer.h"
 #include "Templates/UnrealTemplate.h"
 
+#define UE_API TWEENINGUTILSEDITOR_API
+
 class FUICommandInfo;
 class FUICommandList;
 class ICursor;
@@ -35,7 +37,7 @@ class FTweenModel;
  *		- to the right by 125 Slate units (exceeds the max sliding by 25 units) > OnSliderMovedDelegate is invoked with 1.0 (clamped).
  *		- to the left by 25 Slate units > OnSliderMovedDelegate is invoked with -0.25.
  */
-class TWEENINGUTILSEDITOR_API FMouseSlidingController : public FNoncopyable
+class FMouseSlidingController : public FNoncopyable
 {
 	friend FMouseSlidingInputProcessor;
 public:
@@ -45,12 +47,12 @@ public:
 	 * @param InCommandList The command list to add / remove the command to / from.
 	 * @param InDragSliderCommand The command that triggers detection of mouse movement.
 	 */
-	FMouseSlidingController(
+	UE_API FMouseSlidingController(
 		TAttribute<float> InMaxSlideWidthAttr,
 		const TSharedRef<FUICommandList>& InCommandList,
 		TSharedPtr<FUICommandInfo> InDragSliderCommand
 		);
-	virtual ~FMouseSlidingController();
+	UE_API virtual ~FMouseSlidingController();
 
 	DECLARE_MULTICAST_DELEGATE_OneParam(FMoveSliderDelegate, float /*InNormalizedRange*/);
 	/** Invoked every tick the slider is moved. The event is throttled, so there's at most one invocation per tick. */
@@ -98,17 +100,19 @@ private:
 	FSimpleMulticastDelegate OnStopSlidingDelegate;
 
 	/** Once DragSliderCommand has triggered, start listening for mouse down events. */
-	void StartListeningForMouseEvents();
+	UE_API void StartListeningForMouseEvents();
 	/** Once DragSliderCommand is released, stop listening for mouse down events.  */
-	void StopListeningForMouseEvents(ICursor& InCursor);
+	UE_API void StopListeningForMouseEvents(ICursor& InCursor);
 
 	/** Starts sliding from mouse position and locks the mouse; emits OnUpdateSliding calls. Called once DragSliderCommand and LMB are both down. */
-	void StartSliding(const FVector2D& InInitialScreenLocation, ICursor& InCursor);
+	UE_API void StartSliding(const FVector2D& InInitialScreenLocation, ICursor& InCursor);
 	/** Stops emiting emitting OnUpdateSliding calls. Unlocks the cursor. Called once either DragSliderCommand or LMB are released. */
-	void StopSliding(ICursor& InCursor);
+	UE_API void StopSliding(ICursor& InCursor);
 	/** Emits a OnUpdateSliding call. Called once per frame while sliding is active. */
-	void UpdateSliding(const FVector2D& InScreenLocation);
+	UE_API void UpdateSliding(const FVector2D& InScreenLocation);
 };
 
 	
 }
+
+#undef UE_API
diff --git a/Engine/Plugins/Animation/TweeningUtils/Source/TweeningUtilsEditor/Public/Widgets/MVC/STweenView.h b/Engine/Plugins/Animation/TweeningUtils/Source/TweeningUtilsEditor/Public/Widgets/MVC/STweenView.h
--- a/Engine/Plugins/Animation/TweeningUtils/Source/TweeningUtilsEditor/Public/Widgets/MVC/STweenView.h
+++ b/Engine/Plugins/Animation/TweeningUtils/Source/TweeningUtilsEditor/Public/Widgets/MVC/STweenView.h
@@ -5,6 +5,8 @@
 #include "Widgets/DeclarativeSyntaxSupport.h"
 #include "Widgets/SCompoundWidget.h"
 
+#define UE_API TWEENINGUTILSEDITOR_API
+
 template<typename OptionalType> struct TOptional;
 
 namespace UE::TweeningUtilsEditor
@@ -16,7 +18,7 @@ enum class ETweenScaleMode : uint8;
  * It bridges the STweenSlider, which acts as the view, and the model, FTweenModel, by interchanging data between them.
  * Acts as view in a Model-View-Controller architecture.
  */
-class TWEENINGUTILSEDITOR_API STweenView : public SCompoundWidget
+class STweenView : public SCompoundWidget
 {
 public:
 
@@ -32,19 +34,21 @@ public:
 		SLATE_ATTRIBUTE(TOptional<float>, OverrideSliderPosition)
 	SLATE_END_ARGS()
 
-	void Construct(const FArguments& InArgs);
+	UE_API void Construct(const FArguments& InArgs);
 
 private:
 
 	/** The model that these controls are driving */
 	TAttribute<FTweenModel*> TweenModelAttr;
 	
-	ETweenScaleMode GetBarRenderMode() const;
+	UE_API ETweenScaleMode GetBarRenderMode() const;
 	
-	void OnDragStarted() const;
-	void OnDragEnded() const;
-	void OnDragValueUpdated(float Value) const;
-	void OnPointPicked(float Value) const;
-	float MapSliderValueToBlendValue(float Value) const;
+	UE_API void OnDragStarted() const;
+	UE_API void OnDragEnded() const;
+	UE_API void OnDragValueUpdated(float Value) const;
+	UE_API void OnPointPicked(float Value) const;
+	UE_API float MapSliderValueToBlendValue(float Value) const;
 };
 }
+
+#undef UE_API
diff --git a/Engine/Plugins/Animation/TweeningUtils/Source/TweeningUtilsEditor/Public/Widgets/MVC/TweenControllers.h b/Engine/Plugins/Animation/TweeningUtils/Source/TweeningUtilsEditor/Public/Widgets/MVC/TweenControllers.h
--- a/Engine/Plugins/Animation/TweeningUtils/Source/TweeningUtilsEditor/Public/Widgets/MVC/TweenControllers.h
+++ b/Engine/Plugins/Animation/TweeningUtils/Source/TweeningUtilsEditor/Public/Widgets/MVC/TweenControllers.h
@@ -7,12 +7,14 @@
 #include "TweenMouseSlidingController.h"
 #include "TweenToolbarController.h"
 
+#define UE_API TWEENINGUTILSEDITOR_API
+
 class FUICommandList;
 
 namespace UE::TweeningUtilsEditor
 {
 /** Holds functionality that may be reused across modules to add tweening widgets to a toolbar. */
-struct TWEENINGUTILSEDITOR_API FTweenControllers
+struct FTweenControllers
 {
 	/** Manages the toolbar widget */
 	FTweenToolbarController ToolbarController;
@@ -28,7 +30,7 @@ struct TWEENINGUTILSEDITOR_API FTweenControllers
 	 * @param InTweenModels Contains all tween models that can be selected by the tweening widgets
 	 * @param UserPreferredFunctionContext Optional. The config key under which the selected tween function is saved. The function currently saved under this key will be the first one to be selected.
 	 */
-	explicit FTweenControllers(
+	UE_API explicit FTweenControllers(
 		const TSharedRef<FUICommandList>& InCommandList,
 		const TSharedRef<ITweenModelContainer>& InTweenModels,
 		FName UserPreferredFunctionContext = NAME_None
@@ -36,3 +38,4 @@ struct TWEENINGUTILSEDITOR_API FTweenControllers
 };
 }
 
+#undef UE_API
diff --git a/Engine/Plugins/Animation/TweeningUtils/Source/TweeningUtilsEditor/Public/Widgets/MVC/TweenMouseSlidingController.h b/Engine/Plugins/Animation/TweeningUtils/Source/TweeningUtilsEditor/Public/Widgets/MVC/TweenMouseSlidingController.h
--- a/Engine/Plugins/Animation/TweeningUtils/Source/TweeningUtilsEditor/Public/Widgets/MVC/TweenMouseSlidingController.h
+++ b/Engine/Plugins/Animation/TweeningUtils/Source/TweeningUtilsEditor/Public/Widgets/MVC/TweenMouseSlidingController.h
@@ -6,6 +6,8 @@
 #include "MouseSlidingController.h"
 #include "Templates/SharedPointer.h"
 
+#define UE_API TWEENINGUTILSEDITOR_API
+
 class FUICommandInfo;
 class FUICommandList;
 
@@ -14,7 +16,7 @@ namespace UE::TweeningUtilsEditor
 class FTweenModel;
 
 /** Specialization that sends blend values to a tween model based on how far the mouse is moved. */
-class TWEENINGUTILSEDITOR_API FTweenMouseSlidingController : public FMouseSlidingController
+class FTweenMouseSlidingController : public FMouseSlidingController
 {
 public:
 	
@@ -25,7 +27,7 @@ public:
 	 * @param InCommandList The command list to bind / unbind the command with.
 	 * @param InDragSliderCommand The command that triggers detection of mouse movement.
 	 */
-	explicit FTweenMouseSlidingController(
+	UE_API explicit FTweenMouseSlidingController(
 		TAttribute<float> InMaxSlideWidthAttr,
 		TAttribute<FTweenModel*> InTweenModelAttr,
 		const TSharedRef<FUICommandList>& InCommandList,
@@ -35,7 +37,7 @@ public:
 	 * Version that defaults to using the commands from FTweeningUtilsCommands.
 	 * The slider width is based on the size of the slider in the toolbar(= value from FTweeningUtilsStyle).
 	 */
-	explicit FTweenMouseSlidingController(
+	UE_API explicit FTweenMouseSlidingController(
 		TAttribute<FTweenModel*> InTweenModelAttr,
 		const TSharedRef<FUICommandList>& InCommandList
 	);
@@ -51,8 +53,10 @@ private:
 	/** The current slider position that is being driven by the mouse. Range [-1,1]. */
 	TOptional<float> CurrentSliderPosition;
 
-	void OnSliderStartMove();
-	void OnSliderEndMove();
-	void OnSliderMove(float InValue);
+	UE_API void OnSliderStartMove();
+	UE_API void OnSliderEndMove();
+	UE_API void OnSliderMove(float InValue);
 };
 }
+
+#undef UE_API
