[MLDeformer] Add support for copy paste between input bone and curve lists. This allows users to copy/paste the input lists from one model to the other. You can now right click the input list widgets and choose Copy and Paste.
Also fixed a few bugs in the tooltips as well as unique curve group name generation.
Improved the UX a bit by setting a minimum height for the bone, curve, bone group and curve group widgets and some context menu improvements.
Fixed some UI refresh issues as well when adding bones or curves. Curve group and bone group coloring did not refresh properly.

diff --git a/Engine/Plugins/Animation/MLDeformer/MLDeformerFramework/Source/MLDeformerFrameworkEditor/MLDeformerFrameworkEditor.Build.cs b/Engine/Plugins/Animation/MLDeformer/MLDeformerFramework/Source/MLDeformerFrameworkEditor/MLDeformerFrameworkEditor.Build.cs
--- a/Engine/Plugins/Animation/MLDeformer/MLDeformerFramework/Source/MLDeformerFrameworkEditor/MLDeformerFrameworkEditor.Build.cs
+++ b/Engine/Plugins/Animation/MLDeformer/MLDeformerFramework/Source/MLDeformerFrameworkEditor/MLDeformerFrameworkEditor.Build.cs
@@ -52,9 +52,10 @@ namespace UnrealBuildTool.Rules
 					"TimeManagement",
 					"RenderCore",
 					"RHI",
+					"Json",
 					"AssetDefinition",
 					"SkeletalMeshDescription",
-
+					"ApplicationCore",
 					"EditorInteractiveToolsFramework",
 					"InteractiveToolsFramework",
 					"MeshModelingTools",
diff --git a/Engine/Plugins/Animation/MLDeformer/MLDeformerFramework/Source/MLDeformerFrameworkEditor/Private/SMLDeformerInputBonesWidget.cpp b/Engine/Plugins/Animation/MLDeformer/MLDeformerFramework/Source/MLDeformerFrameworkEditor/Private/SMLDeformerInputBonesWidget.cpp
--- a/Engine/Plugins/Animation/MLDeformer/MLDeformerFramework/Source/MLDeformerFrameworkEditor/Private/SMLDeformerInputBonesWidget.cpp
+++ b/Engine/Plugins/Animation/MLDeformer/MLDeformerFramework/Source/MLDeformerFrameworkEditor/Private/SMLDeformerInputBonesWidget.cpp
@@ -12,6 +12,10 @@
 #include "SMLDeformerBonePickerDialog.h"
 #include "Engine/SkeletalMesh.h"
 #include "ScopedTransaction.h"
+#include "Dom/JsonObject.h"
+#include "Serialization/JsonWriter.h"
+#include "Serialization/JsonSerializer.h"
+#include "HAL/PlatformApplicationMisc.h"
 
 #define LOCTEXT_NAMESPACE "MLDeformerInputBonesWidget"
 
@@ -55,16 +59,40 @@ namespace UE::MLDeformer
 		FMenuBuilder Menu(true, InputBonesWidget->GetInputWidget()->GetBonesCommandList());
 		Menu.BeginSection("BoneActions", LOCTEXT("BoneActionsHeading", "Bone Actions"));
 		{
+			Menu.AddMenuEntry(Actions.AddInputBones);
+			Menu.AddMenuEntry(Actions.AddAnimatedBones);
+			
 			if (!GetSelectedItems().IsEmpty())
 			{
+				Menu.AddSeparator();
 				Menu.AddMenuEntry(Actions.DeleteInputBones);
 			}
+
+			if (!GetRootElements().IsEmpty())
+			{
+				Menu.AddMenuEntry(Actions.ClearInputBones);
+			}
 		}
 		Menu.EndSection();
 
 		const FMLDeformerEditorModel* EditorModel = InputBonesWidget->EditorModel;
 		InputBonesWidget->GetInputWidget()->AddInputBonesMenuItems(Menu);
 
+		Menu.BeginSection("CopyPasteActions", LOCTEXT("CopyPasteActionsHeading", "Copy/Paste Actions"));
+		{
+			if (!RootElements.IsEmpty())
+			{
+				Menu.AddMenuEntry(Actions.CopyInputBones);
+			}
+
+			TSharedPtr<FJsonObject> JsonObject;
+			if (InputBonesWidget->GetInputWidget()->HasValidClipBoardData(TEXT("InputBones"), JsonObject))
+			{
+				Menu.AddMenuEntry(Actions.PasteInputBones);
+			}
+		}
+		Menu.EndSection();
+		
 		return Menu.MakeWidget();
 	}
 
@@ -285,6 +313,7 @@ namespace UE::MLDeformer
 					})
 			]
 			+SVerticalBox::Slot()
+			.MinHeight(75.0f)
 			[
 				SAssignNew(TreeWidget, SMLDeformerInputBoneTreeWidget)
 				.InputBonesWidget(SharedThis(this))
@@ -307,6 +336,74 @@ namespace UE::MLDeformer
 		CommandList->MapAction(Commands.DeleteInputBones, FExecuteAction::CreateSP(this, &SMLDeformerInputBonesWidget::OnDeleteInputBones));
 		CommandList->MapAction(Commands.ClearInputBones, FExecuteAction::CreateSP(this, &SMLDeformerInputBonesWidget::OnClearInputBones));
 		CommandList->MapAction(Commands.AddAnimatedBones, FExecuteAction::CreateSP(this, &SMLDeformerInputBonesWidget::OnAddAnimatedBones));
+		CommandList->MapAction(Commands.CopyInputBones, FExecuteAction::CreateSP(this, &SMLDeformerInputBonesWidget::OnCopyInputBones));
+		CommandList->MapAction(Commands.PasteInputBones, FExecuteAction::CreateSP(this, &SMLDeformerInputBonesWidget::OnPasteInputBones));
+	}
+
+	void SMLDeformerInputBonesWidget::OnCopyInputBones() const
+	{
+		const TSharedRef<FJsonObject> JsonObject = MakeShared<FJsonObject>();
+		JsonObject->SetStringField(TEXT("DataId"), TEXT("InputBones"));
+
+		// Build the names array.
+		TArray<TSharedPtr<FJsonValue>> JsonNameArray;
+		for (const FName Name : EditorModel->GetEditorInputInfo()->GetBoneNames())
+		{
+			if (Name.IsValid() && !Name.IsNone())
+			{
+				JsonNameArray.Add(MakeShared<FJsonValueString>(Name.ToString()));
+			}
+		}		
+		JsonObject->SetArrayField(TEXT("Names"), JsonNameArray);
+
+		// Convert the JSON to a string and copy it to the clipboard.
+		FString OutputString;
+		const TSharedRef<TJsonWriter<>> Writer = TJsonWriterFactory<>::Create(&OutputString);
+		if (!FJsonSerializer::Serialize(JsonObject, Writer))
+		{
+			UE_LOG(LogMLDeformer, Error, TEXT("Failed to serialize JSON on Copy action in the Input Bones widget."));
+			return;
+		}
+		FPlatformApplicationMisc::ClipboardCopy(*OutputString);
+	}
+
+	void SMLDeformerInputBonesWidget::OnPasteInputBones()
+	{
+		FString Text;
+		FPlatformApplicationMisc::ClipboardPaste(Text);
+
+		TSharedPtr<FJsonObject> JsonObject;
+		if (GetInputWidget()->HasValidClipBoardData(TEXT("InputBones"), JsonObject) && JsonObject.IsValid())
+		{
+			TArray<FString> Names;
+			if (!JsonObject->TryGetStringArrayField(TEXT("Names"), Names))
+			{
+				UE_LOG(LogMLDeformer, Error, TEXT("Failed to get bone names during paste operation."));
+			}
+				
+			FScopedTransaction Transaction(TEXT("SMLDeformerInputBonesWidget"), LOCTEXT("PasteInputBonesText", "Paste Input Bones"), EditorModel->GetModel());
+			EditorModel->GetModel()->Modify();
+		
+			TArray<FName> AddedBones;
+			TArray<FBoneReference>& CurrentBoneNames = EditorModel->GetModel()->GetBoneIncludeList();
+			for (const FString& BoneNameString : Names)
+			{
+				if (!BoneNameString.IsEmpty())
+				{
+					const FName BoneName(BoneNameString);
+					CurrentBoneNames.AddUnique(BoneName);
+					AddedBones.Add(BoneName);
+				}
+			}
+			
+			EditorModel->UpdateEditorInputInfo();
+			Refresh();
+
+			// Trigger the input widgets events.
+			// This is done AFTER the RefreshTree call, because that updates the editor input info.
+			// Some handler code might depend on that to be updated first.
+			GetInputWidget()->OnAddInputBones(AddedBones);
+		}
 	}
 
 	void SMLDeformerInputBonesWidget::OnAddInputBones()
@@ -366,6 +463,7 @@ namespace UE::MLDeformer
 		FScopedTransaction Transaction(LOCTEXT("ClearInputBonesText", "Clear Input Bones"));
 		Model->Modify();
 		Model->GetBoneIncludeList().Empty();
+		EditorModel->UpdateEditorInputInfo();
 		RefreshTree();
 		TreeWidget->ClearSelection();
 		InputWidget->OnClearInputBones();
@@ -393,6 +491,7 @@ namespace UE::MLDeformer
 			BoneNamesToRemove.Add(Item->Name);
 		}
 
+		EditorModel->UpdateEditorInputInfo();
 		RefreshTree();
 		TreeWidget->ClearSelection();
 
@@ -408,6 +507,7 @@ namespace UE::MLDeformer
 
 		EditorModel->AddAnimatedBonesToBonesIncludeList();
 		EditorModel->SetResamplingInputOutputsNeeded(true);
+		EditorModel->UpdateEditorInputInfo();
 		RefreshTree();
 		InputWidget->OnAddAnimatedBones();
 	}
@@ -469,6 +569,8 @@ namespace UE::MLDeformer
 		UI_COMMAND(DeleteInputBones, "Delete Selected", "Deletes the selected input bones.", EUserInterfaceActionType::Button, FInputChord(EKeys::Delete));
 		UI_COMMAND(ClearInputBones, "Clear List", "Clears the entire list of input bones.", EUserInterfaceActionType::Button, FInputChord());
 		UI_COMMAND(AddAnimatedBones, "Add All Animated Bones", "Add all animated bones to the list.", EUserInterfaceActionType::Button, FInputChord());
+		UI_COMMAND(CopyInputBones, "Copy Input Bones List", "Copies all bones, so that you can paste them into another model's input bones list.", EUserInterfaceActionType::Button, FInputChord(EKeys::C, EModifierKey::Control));
+		UI_COMMAND(PasteInputBones, "Paste Input Bones List", "Pastes a previously copied bone list. This will add bones to the existing list, not replace it fully.", EUserInterfaceActionType::Button, FInputChord(EKeys::V, EModifierKey::Control));
 	}
 
 	FText SMLDeformerInputBonesWidget::GetSectionTitle() const
diff --git a/Engine/Plugins/Animation/MLDeformer/MLDeformerFramework/Source/MLDeformerFrameworkEditor/Private/SMLDeformerInputCurvesWidget.cpp b/Engine/Plugins/Animation/MLDeformer/MLDeformerFramework/Source/MLDeformerFrameworkEditor/Private/SMLDeformerInputCurvesWidget.cpp
--- a/Engine/Plugins/Animation/MLDeformer/MLDeformerFramework/Source/MLDeformerFrameworkEditor/Private/SMLDeformerInputCurvesWidget.cpp
+++ b/Engine/Plugins/Animation/MLDeformer/MLDeformerFramework/Source/MLDeformerFrameworkEditor/Private/SMLDeformerInputCurvesWidget.cpp
@@ -14,6 +14,11 @@
 #include "Animation/Skeleton.h"
 #include "ScopedTransaction.h"
 #include "Animation/AnimCurveMetadata.h"
+#include "Dom/JsonObject.h"
+#include "Serialization/JsonWriter.h"
+#include "Serialization/JsonSerializer.h"
+#include "HAL/PlatformApplicationMisc.h"
+
 
 #define LOCTEXT_NAMESPACE "MLDeformerInputCurvesWidget"
 
@@ -51,16 +56,41 @@ namespace UE::MLDeformer
 		FMenuBuilder Menu(true, InputCurvesWidget->GetInputWidget()->GetCurvesCommandList());
 		Menu.BeginSection("CurveActions", LOCTEXT("CurveActionsHeading", "Curve Actions"));
 		{
+			Menu.AddMenuEntry(Actions.AddInputCurves);
+			Menu.AddMenuEntry(Actions.AddAnimatedCurves);
+			
 			if (!GetSelectedItems().IsEmpty())
 			{
+				Menu.AddSeparator();
 				Menu.AddMenuEntry(Actions.DeleteInputCurves);
 			}
+
+			if (!GetElements().IsEmpty())
+			{
+				Menu.AddMenuEntry(Actions.ClearInputCurves);
+			}
 		}
 		Menu.EndSection();
 
 		const FMLDeformerEditorModel* EditorModel = InputCurvesWidget->EditorModel;
 		InputCurvesWidget->GetInputWidget()->AddInputCurvesMenuItems(Menu);
 
+		Menu.BeginSection("CopyPasteActions", LOCTEXT("CopyPasteActionsHeading", "Copy/Paste Actions"));
+		{
+			if (!Elements.IsEmpty())
+			{
+				Menu.AddMenuEntry(Actions.CopyInputCurves);
+			}
+
+			TSharedPtr<FJsonObject> JsonObject;
+			if (InputCurvesWidget->GetInputWidget()->HasValidClipBoardData(TEXT("InputCurves"), JsonObject))
+			{
+				Menu.AddMenuEntry(Actions.PasteInputCurves);
+			}
+		}
+		Menu.EndSection();
+
+		
 		return Menu.MakeWidget();
 	}
 
@@ -208,6 +238,7 @@ namespace UE::MLDeformer
 					})
 			]
 			+SVerticalBox::Slot()
+			.MinHeight(75.0f)
 			[
 				SAssignNew(ListWidget, SMLDeformerInputCurveListWidget)
 				.InputCurvesWidget(SharedThis(this))
@@ -230,6 +261,75 @@ namespace UE::MLDeformer
 		CommandList->MapAction(Commands.DeleteInputCurves, FExecuteAction::CreateSP(this, &SMLDeformerInputCurvesWidget::OnDeleteInputCurves));
 		CommandList->MapAction(Commands.ClearInputCurves, FExecuteAction::CreateSP(this, &SMLDeformerInputCurvesWidget::OnClearInputCurves));
 		CommandList->MapAction(Commands.AddAnimatedCurves, FExecuteAction::CreateSP(this, &SMLDeformerInputCurvesWidget::OnAddAnimatedCurves));
+		CommandList->MapAction(Commands.CopyInputCurves, FExecuteAction::CreateSP(this, &SMLDeformerInputCurvesWidget::OnCopyInputCurves));
+		CommandList->MapAction(Commands.PasteInputCurves, FExecuteAction::CreateSP(this, &SMLDeformerInputCurvesWidget::OnPasteInputCurves));
+	}
+
+
+	void SMLDeformerInputCurvesWidget::OnCopyInputCurves() const
+	{
+		const TSharedRef<FJsonObject> JsonObject = MakeShared<FJsonObject>();
+		JsonObject->SetStringField(TEXT("DataId"), TEXT("InputCurves"));
+
+		// Build the names array.
+		TArray<TSharedPtr<FJsonValue>> JsonNameArray;
+		for (const FName Name : EditorModel->GetEditorInputInfo()->GetCurveNames())
+		{
+			if (Name.IsValid() && !Name.IsNone())
+			{
+				JsonNameArray.Add(MakeShared<FJsonValueString>(Name.ToString()));
+			}
+		}		
+		JsonObject->SetArrayField(TEXT("Names"), JsonNameArray);
+
+		// Convert the JSON to a string and copy it to the clipboard.
+		FString OutputString;
+		const TSharedRef<TJsonWriter<>> Writer = TJsonWriterFactory<>::Create(&OutputString);
+		if (!FJsonSerializer::Serialize(JsonObject, Writer))
+		{
+			UE_LOG(LogMLDeformer, Error, TEXT("Failed to serialize JSON on Copy action in the Input Curves widget."));
+			return;
+		}
+		FPlatformApplicationMisc::ClipboardCopy(*OutputString);
+	}
+
+	void SMLDeformerInputCurvesWidget::OnPasteInputCurves()
+	{
+		FString Text;
+		FPlatformApplicationMisc::ClipboardPaste(Text);
+
+		TSharedPtr<FJsonObject> JsonObject;
+		if (GetInputWidget()->HasValidClipBoardData(TEXT("InputCurves"), JsonObject) && JsonObject.IsValid())
+		{
+			TArray<FString> Names;
+			if (!JsonObject->TryGetStringArrayField(TEXT("Names"), Names))
+			{
+				UE_LOG(LogMLDeformer, Error, TEXT("Failed to get curve names during paste operation."));
+			}
+				
+			FScopedTransaction Transaction(TEXT("SMLDeformerInputCurvesWidget"), LOCTEXT("PasteInputCurvesText", "Paste Input Curves"), EditorModel->GetModel());
+			EditorModel->GetModel()->Modify();
+		
+			TArray<FName> AddedCurves;
+			TArray<FMLDeformerCurveReference>& CurrentCurveNames = EditorModel->GetModel()->GetCurveIncludeList();
+			for (const FString& CurveNameString : Names)
+			{
+				if (!CurveNameString.IsEmpty())
+				{
+					const FName CurveName(CurveNameString);
+					CurrentCurveNames.AddUnique(CurveName);
+					AddedCurves.Add(CurveName);
+				}
+			}
+		
+			EditorModel->UpdateEditorInputInfo();
+			Refresh();
+
+			// Trigger the input widgets events.
+			// This is done AFTER the RefreshTree call, because that updates the editor input info.
+			// Some handler code might depend on that to be updated first.
+			GetInputWidget()->OnAddInputCurves(AddedCurves);
+		}
 	}
 
 	void SMLDeformerInputCurvesWidget::OnAddInputCurves()
@@ -297,6 +397,7 @@ namespace UE::MLDeformer
 		Model->Modify();
 
 		Model->GetCurveIncludeList().Empty();
+		EditorModel->UpdateEditorInputInfo();
 		RefreshList();
 		ListWidget->ClearSelection();
 		InputWidget->OnClearInputCurves();
@@ -324,6 +425,7 @@ namespace UE::MLDeformer
 			CurvesToRemove.Add(Item->Name);
 		}
 
+		EditorModel->UpdateEditorInputInfo();
 		RefreshList();
 		ListWidget->ClearSelection();
 
@@ -339,6 +441,7 @@ namespace UE::MLDeformer
 
 		EditorModel->AddAnimatedCurvesToCurvesIncludeList();
 		EditorModel->SetResamplingInputOutputsNeeded(true);
+		EditorModel->UpdateEditorInputInfo();
 		RefreshList();
 		InputWidget->OnAddAnimatedCurves();
 	}
@@ -398,6 +501,8 @@ namespace UE::MLDeformer
 		UI_COMMAND(DeleteInputCurves, "Delete Selected", "Deletes the selected input curves.", EUserInterfaceActionType::Button, FInputChord(EKeys::Delete));
 		UI_COMMAND(ClearInputCurves, "Clear List", "Clears the entire list of input curves.", EUserInterfaceActionType::Button, FInputChord());
 		UI_COMMAND(AddAnimatedCurves, "Add All Animated Curves", "Add all animated curves to the list.", EUserInterfaceActionType::Button, FInputChord());
+		UI_COMMAND(CopyInputCurves, "Copy Input Curves List", "Copies all curves, so that you can paste them into another model's input curves list.", EUserInterfaceActionType::Button, FInputChord(EKeys::C, EModifierKey::Control));
+		UI_COMMAND(PasteInputCurves, "Paste Input Curves List", "Pastes a previously copied curves list. This will add curves to the existing list, not replace it fully.", EUserInterfaceActionType::Button, FInputChord(EKeys::V, EModifierKey::Control));
 	}
 
 	FText SMLDeformerInputCurvesWidget::GetSectionTitle() const
diff --git a/Engine/Plugins/Animation/MLDeformer/MLDeformerFramework/Source/MLDeformerFrameworkEditor/Private/SMLDeformerInputWidget.cpp b/Engine/Plugins/Animation/MLDeformer/MLDeformerFramework/Source/MLDeformerFrameworkEditor/Private/SMLDeformerInputWidget.cpp
--- a/Engine/Plugins/Animation/MLDeformer/MLDeformerFramework/Source/MLDeformerFrameworkEditor/Private/SMLDeformerInputWidget.cpp
+++ b/Engine/Plugins/Animation/MLDeformer/MLDeformerFramework/Source/MLDeformerFrameworkEditor/Private/SMLDeformerInputWidget.cpp
@@ -23,6 +23,9 @@
 #include "IDetailsView.h"
 #include "Misc/Attribute.h"
 #include "Internationalization/Text.h"
+#include "HAL/PlatformApplicationMisc.h"
+#include "Serialization/JsonReader.h"
+#include "Serialization/JsonSerializer.h"
 
 #define LOCTEXT_NAMESPACE "MLDeformerInputWidget"
 
@@ -233,7 +236,57 @@ namespace UE::MLDeformer
 		FMLDeformerInputBonesWidgetCommands::Unregister();
 		FMLDeformerInputCurvesWidgetCommands::Unregister();
 	}
+
+	TSharedPtr<FJsonObject> SMLDeformerInputWidget::GetJsonObject(const FString& JsonString)
+	{
+		TSharedPtr<FJsonObject> JsonObject;
+		const TSharedRef<TJsonReader<>> Reader = TJsonReaderFactory<>::Create(JsonString);
+		if (FJsonSerializer::Deserialize(Reader, JsonObject))
+		{
+			return JsonObject;
+		}
+
+		return TSharedPtr<FJsonObject>();
+	}
+	
+	bool SMLDeformerInputWidget::ExtractStringFromJson(const TSharedPtr<FJsonObject>& JsonObject, const FString& Key, FString& OutString)
+	{
+		if (JsonObject.IsValid())
+		{
+			if (JsonObject->HasTypedField(Key, EJson::String))
+			{
+				OutString = JsonObject->GetStringField(Key);
+				return true;
+			}
+		}
+		return false;
+	}
 	
+	bool SMLDeformerInputWidget::HasValidClipBoardData(const FString& DataTypeId, TSharedPtr<FJsonObject>& OutJsonObject)
+	{
+		// Get the text from the clipboard.
+		FString Text;
+		FPlatformApplicationMisc::ClipboardPaste(Text);
+
+		const TSharedPtr<FJsonObject> JsonObject = GetJsonObject(Text);
+		if (JsonObject.IsValid())
+		{
+			OutJsonObject = JsonObject;
+			
+			// Check if we find the "DataTypeId" field, and get it's value.
+			// Then compare it to the data type id we are interested in.
+			FString DataTypeIdValue;
+			if (ExtractStringFromJson(JsonObject, TEXT("DataId"), DataTypeIdValue))
+			{
+				return DataTypeIdValue.Compare(DataTypeId, ESearchCase::IgnoreCase) == 0;
+			}
+		}
+
+		// Invalid JSON or no DataTypeId key found, or the value didn't match the one we search for.
+		OutJsonObject = JsonObject;
+		return false;
+	}
+
 	TSharedPtr<FUICommandList> SMLDeformerInputWidget::GetBonesCommandList() const
 	{ 
 		return BonesCommandList;
diff --git a/Engine/Plugins/Animation/MLDeformer/MLDeformerFramework/Source/MLDeformerFrameworkEditor/Public/SMLDeformerInputBonesWidget.h b/Engine/Plugins/Animation/MLDeformer/MLDeformerFramework/Source/MLDeformerFrameworkEditor/Public/SMLDeformerInputBonesWidget.h
--- a/Engine/Plugins/Animation/MLDeformer/MLDeformerFramework/Source/MLDeformerFrameworkEditor/Public/SMLDeformerInputBonesWidget.h
+++ b/Engine/Plugins/Animation/MLDeformer/MLDeformerFramework/Source/MLDeformerFrameworkEditor/Public/SMLDeformerInputBonesWidget.h
@@ -34,6 +34,8 @@ namespace UE::MLDeformer
 		TSharedPtr<FUICommandInfo> DeleteInputBones;
 		TSharedPtr<FUICommandInfo> ClearInputBones;
 		TSharedPtr<FUICommandInfo> AddAnimatedBones;
+		TSharedPtr<FUICommandInfo> CopyInputBones;
+		TSharedPtr<FUICommandInfo> PasteInputBones;
 	};
 
 
@@ -121,6 +123,8 @@ namespace UE::MLDeformer
 		UE_API void OnDeleteInputBones();
 		UE_API void OnClearInputBones();
 		UE_API void OnAddAnimatedBones();
+		UE_API void OnCopyInputBones() const;
+		UE_API void OnPasteInputBones();
 
 	private:
 		TSharedPtr<SMLDeformerInputBoneTreeWidget> TreeWidget;
diff --git a/Engine/Plugins/Animation/MLDeformer/MLDeformerFramework/Source/MLDeformerFrameworkEditor/Public/SMLDeformerInputCurvesWidget.h b/Engine/Plugins/Animation/MLDeformer/MLDeformerFramework/Source/MLDeformerFrameworkEditor/Public/SMLDeformerInputCurvesWidget.h
--- a/Engine/Plugins/Animation/MLDeformer/MLDeformerFramework/Source/MLDeformerFrameworkEditor/Public/SMLDeformerInputCurvesWidget.h
+++ b/Engine/Plugins/Animation/MLDeformer/MLDeformerFramework/Source/MLDeformerFrameworkEditor/Public/SMLDeformerInputCurvesWidget.h
@@ -34,6 +34,8 @@ namespace UE::MLDeformer
 		TSharedPtr<FUICommandInfo> DeleteInputCurves;
 		TSharedPtr<FUICommandInfo> ClearInputCurves;
 		TSharedPtr<FUICommandInfo> AddAnimatedCurves;
+		TSharedPtr<FUICommandInfo> CopyInputCurves;
+		TSharedPtr<FUICommandInfo> PasteInputCurves;	
 	};
 
 
@@ -112,6 +114,9 @@ namespace UE::MLDeformer
 		UE_API void RefreshList(bool bBroadCastPropertyChanged=true);
 		UE_API bool BroadcastModelPropertyChanged(const FName PropertyName);
 
+		UE_API void OnCopyInputCurves() const;
+		UE_API void OnPasteInputCurves();
+	
 		UE_API void OnAddInputCurves();
 		UE_API void OnDeleteInputCurves();
 		UE_API void OnClearInputCurves();
diff --git a/Engine/Plugins/Animation/MLDeformer/MLDeformerFramework/Source/MLDeformerFrameworkEditor/Public/SMLDeformerInputWidget.h b/Engine/Plugins/Animation/MLDeformer/MLDeformerFramework/Source/MLDeformerFrameworkEditor/Public/SMLDeformerInputWidget.h
--- a/Engine/Plugins/Animation/MLDeformer/MLDeformerFramework/Source/MLDeformerFrameworkEditor/Public/SMLDeformerInputWidget.h
+++ b/Engine/Plugins/Animation/MLDeformer/MLDeformerFramework/Source/MLDeformerFrameworkEditor/Public/SMLDeformerInputWidget.h
@@ -13,6 +13,7 @@
 class SWidget;
 class FMenuBuilder;
 class SVerticalBox;
+class FJsonObject;
 
 namespace UE::MLDeformer
 {
@@ -71,6 +72,10 @@ namespace UE::MLDeformer
 		static UE_API void RegisterCommands();
 		static UE_API void UnregisterCommands();
 
+		static UE_API TSharedPtr<FJsonObject> GetJsonObject(const FString& JsonString);
+		static UE_API bool ExtractStringFromJson(const TSharedPtr<FJsonObject>& JsonObject, const FString& Key, FString& OutString);
+		static UE_API bool HasValidClipBoardData(const FString& DataTypeId, TSharedPtr<FJsonObject>& OutJsonObject);
+
 	protected:
 		UE_API void AddSectionSeparator();
 		UE_API void CreateBonesWidget();
diff --git a/Engine/Plugins/Animation/MLDeformer/NeuralMorphModel/Source/NeuralMorphModelEditor/NeuralMorphModelEditor.Build.cs b/Engine/Plugins/Animation/MLDeformer/NeuralMorphModel/Source/NeuralMorphModelEditor/NeuralMorphModelEditor.Build.cs
--- a/Engine/Plugins/Animation/MLDeformer/NeuralMorphModel/Source/NeuralMorphModelEditor/NeuralMorphModelEditor.Build.cs
+++ b/Engine/Plugins/Animation/MLDeformer/NeuralMorphModel/Source/NeuralMorphModelEditor/NeuralMorphModelEditor.Build.cs
@@ -29,7 +29,9 @@ namespace UnrealBuildTool.Rules
 					"NeuralMorphModel",
 					"PropertyEditor",
 					"ToolWidgets",
-					"DeveloperSettings"
+					"DeveloperSettings",
+					"Json",
+					"ApplicationCore"
 				}
 			);
 
diff --git a/Engine/Plugins/Animation/MLDeformer/NeuralMorphModel/Source/NeuralMorphModelEditor/Private/NeuralMorphEditorModel.cpp b/Engine/Plugins/Animation/MLDeformer/NeuralMorphModel/Source/NeuralMorphModelEditor/Private/NeuralMorphEditorModel.cpp
--- a/Engine/Plugins/Animation/MLDeformer/NeuralMorphModel/Source/NeuralMorphModelEditor/Private/NeuralMorphEditorModel.cpp
+++ b/Engine/Plugins/Animation/MLDeformer/NeuralMorphModel/Source/NeuralMorphModelEditor/Private/NeuralMorphEditorModel.cpp
@@ -403,7 +403,7 @@ namespace UE::NeuralMorphModel
 
 	FName FNeuralMorphEditorModel::GenerateUniqueBoneGroupName() const
 	{
-		UNeuralMorphModel* NeuralMorphModel = GetNeuralMorphModel();
+		const UNeuralMorphModel* NeuralMorphModel = GetNeuralMorphModel();
 		FName GeneratedName;
 		bool bUniqueNameFound = false;
 		int32 Counter = 0;
@@ -432,7 +432,7 @@ namespace UE::NeuralMorphModel
 
 	FName FNeuralMorphEditorModel::GenerateUniqueCurveGroupName() const
 	{
-		UNeuralMorphModel* NeuralMorphModel = GetNeuralMorphModel();
+		const UNeuralMorphModel* NeuralMorphModel = GetNeuralMorphModel();
 		FName GeneratedName;
 		bool bUniqueNameFound = false;
 		int32 Counter = 0;
@@ -441,9 +441,9 @@ namespace UE::NeuralMorphModel
 			GeneratedName = FName(FString::Format(TEXT("Curve Group #{0}"), {Counter++}));
 
 			bool bHasGroupWithSameName = false;
-			for (const FNeuralMorphBoneGroup& BoneGroup : NeuralMorphModel->GetBoneGroups())
+			for (const FNeuralMorphCurveGroup& CurveGroup : NeuralMorphModel->GetCurveGroups())
 			{
-				if (BoneGroup.GroupName == GeneratedName)
+				if (CurveGroup.GroupName == GeneratedName)
 				{
 					bHasGroupWithSameName = true;
 					break;
diff --git a/Engine/Plugins/Animation/MLDeformer/NeuralMorphModel/Source/NeuralMorphModelEditor/Private/SNeuralMorphBoneGroupsWidget.cpp b/Engine/Plugins/Animation/MLDeformer/NeuralMorphModel/Source/NeuralMorphModelEditor/Private/SNeuralMorphBoneGroupsWidget.cpp
--- a/Engine/Plugins/Animation/MLDeformer/NeuralMorphModel/Source/NeuralMorphModelEditor/Private/SNeuralMorphBoneGroupsWidget.cpp
+++ b/Engine/Plugins/Animation/MLDeformer/NeuralMorphModel/Source/NeuralMorphModelEditor/Private/SNeuralMorphBoneGroupsWidget.cpp
@@ -10,6 +10,10 @@
 #include "Framework/MultiBox/MultiBoxBuilder.h"
 #include "Engine/SkeletalMesh.h"
 #include "ScopedTransaction.h"
+#include "Dom/JsonObject.h"
+#include "Serialization/JsonWriter.h"
+#include "Serialization/JsonSerializer.h"
+#include "HAL/PlatformApplicationMisc.h"
 
 #define LOCTEXT_NAMESPACE "NeuralMorphBoneGroupsWidget"
 
@@ -30,6 +34,8 @@ namespace UE::NeuralMorphModel
 		UI_COMMAND(DeleteSelectedItems, "Delete Selected Items", "Deletes the selected bones and/or groups.", EUserInterfaceActionType::Button, FInputChord(EKeys::Delete));
 		UI_COMMAND(ClearGroups, "Clear All Groups", "Clears the entire list of bone groups.", EUserInterfaceActionType::Button, FInputChord());
 		UI_COMMAND(AddBoneToGroup, "Add Bones To Group", "Add new bones to the group.", EUserInterfaceActionType::Button, FInputChord());
+		UI_COMMAND(CopyBoneGroups, "Copy Input Bone Groups List", "Copies all bone groups, so that you can paste them into another model's input bone groups list.", EUserInterfaceActionType::Button, FInputChord(EKeys::C, EModifierKey::Control));
+		UI_COMMAND(PasteBoneGroups, "Paste Input Bones Groups List", "Pastes a previously copied bone group list. This will add bone groups to the existing list, not replace it fully.", EUserInterfaceActionType::Button, FInputChord(EKeys::V, EModifierKey::Control));
 	}
 
 	void SNeuralMorphBoneGroupsWidget::BindCommands(TSharedPtr<FUICommandList> CommandList)
@@ -39,6 +45,8 @@ namespace UE::NeuralMorphModel
 		CommandList->MapAction(GroupCommands.DeleteSelectedItems, FExecuteAction::CreateSP(this, &SNeuralMorphBoneGroupsWidget::OnDeleteSelectedItems));
 		CommandList->MapAction(GroupCommands.ClearGroups, FExecuteAction::CreateSP(this, &SNeuralMorphBoneGroupsWidget::OnClearBoneGroups));
 		CommandList->MapAction(GroupCommands.AddBoneToGroup, FExecuteAction::CreateSP(this, &SNeuralMorphBoneGroupsWidget::OnAddBoneToGroup));
+		CommandList->MapAction(GroupCommands.CopyBoneGroups, FExecuteAction::CreateSP(this, &SNeuralMorphBoneGroupsWidget::OnCopyBoneGroups));
+		CommandList->MapAction(GroupCommands.PasteBoneGroups, FExecuteAction::CreateSP(this, &SNeuralMorphBoneGroupsWidget::OnPasteBoneGroups));
 	}
 
 	void SNeuralMorphBoneGroupsWidget::Construct(const FArguments& InArgs)
@@ -80,6 +88,8 @@ namespace UE::NeuralMorphModel
 		const TArray<TSharedPtr<FNeuralMorphBoneGroupsTreeElement>> CurSelectedItems = GetSelectedItems();
 		Menu.BeginSection("BoneGroupActions", LOCTEXT("BoneGroupActionsHeading", "Bone Group Actions"));
 		{
+			Menu.AddMenuEntry(Actions.CreateGroup);
+			
 			if (CurSelectedItems.Num() == 1 && CurSelectedItems[0]->IsGroup())
 			{
 				Menu.AddMenuEntry(Actions.AddBoneToGroup);
@@ -89,6 +99,11 @@ namespace UE::NeuralMorphModel
 			{			
 				Menu.AddMenuEntry(Actions.DeleteSelectedItems);
 			}
+
+			if (!GetRootElements().IsEmpty())
+			{
+				Menu.AddMenuEntry(Actions.ClearGroups);
+			}
 		}
 		Menu.EndSection();
 
@@ -99,9 +114,136 @@ namespace UE::NeuralMorphModel
 			NeuralInputWidget->AddInputBoneGroupsMenuItems(Menu);
 		}
 
+		Menu.BeginSection("CopyPasteActions", LOCTEXT("CopyPasteActionsHeading", "Copy/Paste Actions"));
+		{
+			if (!RootElements.IsEmpty())
+			{
+				Menu.AddMenuEntry(Actions.CopyBoneGroups);
+			}
+
+			TSharedPtr<FJsonObject> JsonObject;
+			if (InputWidget->HasValidClipBoardData(TEXT("BoneGroups"), JsonObject))
+			{
+				Menu.AddMenuEntry(Actions.PasteBoneGroups);
+			}
+		}
+		Menu.EndSection();
+		
 		return Menu.MakeWidget();
 	}
 
+	void SNeuralMorphBoneGroupsWidget::OnCopyBoneGroups() const
+	{
+		const UNeuralMorphModel* NeuralMorphModel = EditorModel->GetNeuralMorphModel();
+		if (NeuralMorphModel->GetBoneGroups().IsEmpty())
+		{
+			return;
+		}
+
+		const TSharedRef<FJsonObject> JsonObject = MakeShared<FJsonObject>();
+		JsonObject->SetStringField(TEXT("DataId"), TEXT("BoneGroups"));
+
+		// Create each bone group.
+		TArray<TSharedPtr<FJsonValue>> BoneGroupsArray;
+		for (const FNeuralMorphBoneGroup& BoneGroup : NeuralMorphModel->GetBoneGroups())
+		{
+			TSharedRef<FJsonObject> JsonBoneGroup = MakeShared<FJsonObject>();
+			JsonBoneGroup->SetStringField(TEXT("GroupName"), BoneGroup.GroupName.ToString());
+    
+			TArray<TSharedPtr<FJsonValue>> BoneNames;
+			for (const FBoneReference& BoneName : BoneGroup.BoneNames)
+			{
+				if (BoneName.BoneName.IsValid() && !BoneName.BoneName.IsNone())
+				{
+					BoneNames.Add(MakeShared<FJsonValueString>(BoneName.BoneName.ToString()));
+				}
+			}
+	
+			JsonBoneGroup->SetArrayField(TEXT("BoneNames"), BoneNames);
+			BoneGroupsArray.Add(MakeShared<FJsonValueObject>(JsonBoneGroup));
+		}
+
+		JsonObject->SetArrayField(TEXT("BoneGroups"), BoneGroupsArray);
+
+		// Convert the JSON to a string and copy it to the clipboard.
+		FString OutputString;
+		const TSharedRef<TJsonWriter<>> Writer = TJsonWriterFactory<>::Create(&OutputString);
+		if (!FJsonSerializer::Serialize(JsonObject, Writer))
+		{
+			UE_LOG(LogMLDeformer, Error, TEXT("Failed to serialize JSON on Copy action in the Input Bone Groups widget."));
+			return;
+		}
+		FPlatformApplicationMisc::ClipboardCopy(*OutputString);
+	}
+
+	void SNeuralMorphBoneGroupsWidget::OnPasteBoneGroups()
+	{
+		FString Text;
+		FPlatformApplicationMisc::ClipboardPaste(Text);
+
+		TSharedPtr<FJsonObject> JsonObject;
+		if (GetInputWidget()->HasValidClipBoardData(TEXT("BoneGroups"), JsonObject) && JsonObject.IsValid())
+		{
+			const TArray<TSharedPtr<FJsonValue>>* BoneGroupsArray;	
+			if (JsonObject->TryGetArrayField(TEXT("BoneGroups"), BoneGroupsArray))
+			{
+				UNeuralMorphModel* NeuralMorphModel = EditorModel->GetNeuralMorphModel();
+				FScopedTransaction Transaction(TEXT("SNeuralMorphBoneGroupsWidget"), LOCTEXT("PasteBoneGroupText", "Paste Bone Groups"), NeuralMorphModel);
+				NeuralMorphModel->Modify();
+	
+				for (const TSharedPtr<FJsonValue>& GroupValue : *BoneGroupsArray)
+				{
+					TSharedPtr<FJsonObject> GroupObject = GroupValue->AsObject();
+					if (!GroupObject.IsValid())
+					{
+						UE_LOG(LogNeuralMorphModel, Error, TEXT("Skipping invalid group object when pasting bone groups."));
+						continue;
+					}
+
+					FString GroupNameString;
+					if (!GroupObject->TryGetStringField(TEXT("GroupName"), GroupNameString))
+					{
+						UE_LOG(LogNeuralMorphModel, Error, TEXT("Failed to read bone group name when pasting bone groups."));
+						continue;
+					}
+					
+					const TArray<TSharedPtr<FJsonValue>>* BoneArray;
+					if (!GroupObject->TryGetArrayField(TEXT("BoneNames"), BoneArray))
+					{
+						UE_LOG(LogNeuralMorphModel, Error, TEXT("Failed to read bone group bone names array when pasting bone groups."));
+						continue;
+					}
+
+					// Check if the copied group name is unique or if we have to generate a new one.
+					bool bIsUniqueName = true;
+					for (const FNeuralMorphBoneGroup& Group : NeuralMorphModel->BoneGroups)
+					{
+						if (Group.GroupName == GroupNameString)
+						{
+							bIsUniqueName = false;
+							break;
+						}
+					}
+					
+					const FName GroupName = bIsUniqueName ? FName(GroupNameString) : EditorModel->GenerateUniqueBoneGroupName();
+					FNeuralMorphBoneGroup& BoneGroup = NeuralMorphModel->BoneGroups.AddDefaulted_GetRef();
+					BoneGroup.GroupName = GroupName;
+
+					// Add the bones to the newly created group.
+					for (const TSharedPtr<FJsonValue>& BoneValue : *BoneArray)
+					{
+						const FString BoneName = BoneValue->AsString();
+						BoneGroup.BoneNames.AddUnique(FName(BoneName));
+					}
+				}
+			}
+			
+			RefreshTree(true);
+			EditorModel->RebuildEditorMaskInfo();
+		}
+	}
+	
+
 	void SNeuralMorphBoneGroupsWidget::OnCreateBoneGroup()
 	{
 		using namespace UE::MLDeformer;
@@ -144,6 +286,7 @@ namespace UE::NeuralMorphModel
 			BoneGroup.BoneNames.AddUnique(BoneName);
 		}
 
+		EditorModel->UpdateEditorInputInfo();
 		RefreshTree(true);
 		EditorModel->RebuildEditorMaskInfo();
 	}
@@ -220,6 +363,7 @@ namespace UE::NeuralMorphModel
 			}
 		}
 
+		EditorModel->UpdateEditorInputInfo();
 		RefreshTree(true);
 		EditorModel->RebuildEditorMaskInfo();
 	}
@@ -232,6 +376,7 @@ namespace UE::NeuralMorphModel
 
 		NeuralMorphModel->BoneGroups.Empty();
 		NeuralMorphModel->BoneGroupMaskInfoMap.Empty();
+		EditorModel->UpdateEditorInputInfo();
 		RefreshTree(true);
 		EditorModel->RebuildEditorMaskInfo();
 	}
@@ -290,6 +435,7 @@ namespace UE::NeuralMorphModel
 			const int32 HierarchyDepth = StaticCastSharedPtr<SNeuralMorphInputWidget>(InputWidget)->GetHierarchyDepth();
 			EditorModel->GenerateBoneGroupMaskInfo(GroupIndex, HierarchyDepth);
 
+			EditorModel->UpdateEditorInputInfo();
 			RefreshTree(true);
 			EditorModel->RebuildEditorMaskInfo();
 		}
@@ -472,9 +618,9 @@ namespace UE::NeuralMorphModel
 	}
 
 
-	void SNeuralMorphBoneGroupsWidget::Refresh()
+	void SNeuralMorphBoneGroupsWidget::Refresh(bool bBroadcastPropertyChanged)
 	{
-		RefreshTree(false);
+		RefreshTree(bBroadcastPropertyChanged);
 	}
 
 	FText SNeuralMorphBoneGroupsWidget::GetSectionTitle() const
diff --git a/Engine/Plugins/Animation/MLDeformer/NeuralMorphModel/Source/NeuralMorphModelEditor/Private/SNeuralMorphCurveGroupsWidget.cpp b/Engine/Plugins/Animation/MLDeformer/NeuralMorphModel/Source/NeuralMorphModelEditor/Private/SNeuralMorphCurveGroupsWidget.cpp
--- a/Engine/Plugins/Animation/MLDeformer/NeuralMorphModel/Source/NeuralMorphModelEditor/Private/SNeuralMorphCurveGroupsWidget.cpp
+++ b/Engine/Plugins/Animation/MLDeformer/NeuralMorphModel/Source/NeuralMorphModelEditor/Private/SNeuralMorphCurveGroupsWidget.cpp
@@ -10,6 +10,10 @@
 #include "Framework/MultiBox/MultiBoxBuilder.h"
 #include "Engine/SkeletalMesh.h"
 #include "ScopedTransaction.h"
+#include "Dom/JsonObject.h"
+#include "Serialization/JsonWriter.h"
+#include "Serialization/JsonSerializer.h"
+#include "HAL/PlatformApplicationMisc.h"
 
 
 #define LOCTEXT_NAMESPACE "NeuralMorphCurveGroupsWidget"
@@ -27,10 +31,12 @@ namespace UE::NeuralMorphModel
 
 	void FNeuralMorphCurveGroupsCommands::RegisterCommands()
 	{
-		UI_COMMAND(CreateGroup, "Create New Group", "Create a new bone group.", EUserInterfaceActionType::Button, FInputChord(EKeys::Insert));
-		UI_COMMAND(DeleteSelectedItems, "Delete Selected Items", "Deletes the selected bones and/or groups.", EUserInterfaceActionType::Button, FInputChord(EKeys::Delete));
-		UI_COMMAND(ClearGroups, "Clear All Groups", "Clears the entire list of bone groups.", EUserInterfaceActionType::Button, FInputChord());
-		UI_COMMAND(AddCurveToGroup, "Add Curves To Group", "Add new bones to the group.", EUserInterfaceActionType::Button, FInputChord());
+		UI_COMMAND(CreateGroup, "Create New Group", "Create a new curve group.", EUserInterfaceActionType::Button, FInputChord(EKeys::Insert));
+		UI_COMMAND(DeleteSelectedItems, "Delete Selected Items", "Deletes the selected curves and/or groups.", EUserInterfaceActionType::Button, FInputChord(EKeys::Delete));
+		UI_COMMAND(ClearGroups, "Clear All Groups", "Clears the entire list of curve groups.", EUserInterfaceActionType::Button, FInputChord());
+		UI_COMMAND(AddCurveToGroup, "Add Curves To Group", "Add new curves to the group.", EUserInterfaceActionType::Button, FInputChord());
+		UI_COMMAND(CopyCurveGroups, "Copy Input Curve Groups List", "Copies all curve groups, so that you can paste them into another model's input curve groups list.", EUserInterfaceActionType::Button, FInputChord(EKeys::C, EModifierKey::Control));
+		UI_COMMAND(PasteCurveGroups, "Paste Input Curve Groups List", "Pastes a previously copied curve group list. This will add curve groups to the existing list, not replace it fully.", EUserInterfaceActionType::Button, FInputChord(EKeys::V, EModifierKey::Control));
 	}
 
 	void SNeuralMorphCurveGroupsWidget::BindCommands(TSharedPtr<FUICommandList> CommandList)
@@ -40,6 +46,8 @@ namespace UE::NeuralMorphModel
 		CommandList->MapAction(GroupCommands.DeleteSelectedItems, FExecuteAction::CreateSP(this, &SNeuralMorphCurveGroupsWidget::OnDeleteSelectedItems));
 		CommandList->MapAction(GroupCommands.ClearGroups, FExecuteAction::CreateSP(this, &SNeuralMorphCurveGroupsWidget::OnClearCurveGroups));
 		CommandList->MapAction(GroupCommands.AddCurveToGroup, FExecuteAction::CreateSP(this, &SNeuralMorphCurveGroupsWidget::OnAddCurveToGroup));
+		CommandList->MapAction(GroupCommands.CopyCurveGroups, FExecuteAction::CreateSP(this, &SNeuralMorphCurveGroupsWidget::OnCopyCurveGroups));
+		CommandList->MapAction(GroupCommands.PasteCurveGroups, FExecuteAction::CreateSP(this, &SNeuralMorphCurveGroupsWidget::OnPasteCurveGroups));
 	}
 
 	void SNeuralMorphCurveGroupsWidget::Construct(const FArguments& InArgs)
@@ -71,6 +79,117 @@ namespace UE::NeuralMorphModel
 		}
 	}
 
+	void SNeuralMorphCurveGroupsWidget::OnCopyCurveGroups() const
+	{
+		const UNeuralMorphModel* NeuralMorphModel = EditorModel->GetNeuralMorphModel();
+		if (NeuralMorphModel->GetCurveGroups().IsEmpty())
+		{
+			return;
+		}
+
+		const TSharedRef<FJsonObject> JsonObject = MakeShared<FJsonObject>();
+		JsonObject->SetStringField(TEXT("DataId"), TEXT("CurveGroups"));
+
+		// Create each curve group.
+		TArray<TSharedPtr<FJsonValue>> CurveGroupsArray;
+		for (const FNeuralMorphCurveGroup& CurveGroup : NeuralMorphModel->GetCurveGroups())
+		{
+			TSharedRef<FJsonObject> JsonCurveroup = MakeShared<FJsonObject>();
+			JsonCurveroup->SetStringField(TEXT("GroupName"), CurveGroup.GroupName.ToString());
+    
+			TArray<TSharedPtr<FJsonValue>> CurveNames;
+			for (const FMLDeformerCurveReference& CurveName : CurveGroup.CurveNames)
+			{
+				if (CurveName.CurveName.IsValid() && !CurveName.CurveName.IsNone())
+				{
+					CurveNames.Add(MakeShared<FJsonValueString>(CurveName.CurveName.ToString()));
+				}
+			}
+	
+			JsonCurveroup->SetArrayField(TEXT("CurveNames"), CurveNames);
+			CurveGroupsArray.Add(MakeShared<FJsonValueObject>(JsonCurveroup));
+		}
+
+		JsonObject->SetArrayField(TEXT("CurveGroups"), CurveGroupsArray);
+
+		// Convert the JSON to a string and copy it to the clipboard.
+		FString OutputString;
+		const TSharedRef<TJsonWriter<>> Writer = TJsonWriterFactory<>::Create(&OutputString);
+		if (!FJsonSerializer::Serialize(JsonObject, Writer))
+		{
+			UE_LOG(LogMLDeformer, Error, TEXT("Failed to serialize JSON on Copy action in the Input Curve Groups widget."));
+			return;
+		}
+		FPlatformApplicationMisc::ClipboardCopy(*OutputString);
+	}
+
+	void SNeuralMorphCurveGroupsWidget::OnPasteCurveGroups()
+	{
+		FString Text;
+		FPlatformApplicationMisc::ClipboardPaste(Text);
+
+		TSharedPtr<FJsonObject> JsonObject;
+		if (GetInputWidget()->HasValidClipBoardData(TEXT("CurveGroups"), JsonObject) && JsonObject.IsValid())
+		{
+			const TArray<TSharedPtr<FJsonValue>>* CurveGroupsArray;	
+			if (JsonObject->TryGetArrayField(TEXT("CurveGroups"), CurveGroupsArray))
+			{
+				UNeuralMorphModel* NeuralMorphModel = EditorModel->GetNeuralMorphModel();
+				FScopedTransaction Transaction(TEXT("SNeuralMorphCurveGroupsWidget"), LOCTEXT("PasteCurveGroupText", "Paste Curve Groups"), NeuralMorphModel);
+				NeuralMorphModel->Modify();
+	
+				for (const TSharedPtr<FJsonValue>& GroupValue : *CurveGroupsArray)
+				{
+					TSharedPtr<FJsonObject> GroupObject = GroupValue->AsObject();
+					if (!GroupObject.IsValid())
+					{
+						UE_LOG(LogNeuralMorphModel, Error, TEXT("Skipping invalid group object when pasting curve groups."));
+						continue;
+					}
+
+					FString GroupNameString;
+					if (!GroupObject->TryGetStringField(TEXT("GroupName"), GroupNameString))
+					{
+						UE_LOG(LogNeuralMorphModel, Error, TEXT("Failed to read curve group name when pasting curve groups."));
+						continue;
+					}
+					
+					const TArray<TSharedPtr<FJsonValue>>* CurveArray;
+					if (!GroupObject->TryGetArrayField(TEXT("CurveNames"), CurveArray))
+					{
+						UE_LOG(LogNeuralMorphModel, Error, TEXT("Failed to read curve group curve names array when pasting curve groups."));
+						continue;
+					}
+
+					// Check if the copied group name is unique or if we have to generate a new one.
+					bool bIsUniqueName = true;
+					for (const FNeuralMorphCurveGroup& Group : NeuralMorphModel->CurveGroups)
+					{
+						if (Group.GroupName == GroupNameString)
+						{
+							bIsUniqueName = false;
+							break;
+						}
+					}
+					
+					const FName GroupName = bIsUniqueName ? FName(GroupNameString) : EditorModel->GenerateUniqueCurveGroupName();
+					FNeuralMorphCurveGroup& CurveGroup = NeuralMorphModel->CurveGroups.AddDefaulted_GetRef();
+					CurveGroup.GroupName = GroupName;
+
+					// Add the curves to the newly created group.
+					for (const TSharedPtr<FJsonValue>& CurveNameValue : *CurveArray)
+					{
+						const FString CurveName = CurveNameValue->AsString();
+						CurveGroup.CurveNames.AddUnique(FName(CurveName));
+					}
+				}
+			}
+			
+			RefreshTree(true);
+			EditorModel->RebuildEditorMaskInfo();
+		}
+	}
+	
 	TSharedPtr<SWidget> SNeuralMorphCurveGroupsWidget::CreateContextMenuWidget() const
 	{
 		const FNeuralMorphCurveGroupsCommands& Actions = FNeuralMorphCurveGroupsCommands::Get();
@@ -81,6 +200,8 @@ namespace UE::NeuralMorphModel
 		const TArray<TSharedPtr<FNeuralMorphCurveGroupsTreeElement>> CurSelectedItems = GetSelectedItems();
 
 		Menu.BeginSection("CurveGroupActions", LOCTEXT("CurveGroupActionsHeading", "Curve Group Actions"));
+			Menu.AddMenuEntry(Actions.CreateGroup);
+		
 			if (CurSelectedItems.Num() == 1 && CurSelectedItems[0]->IsGroup())
 			{
 				Menu.AddMenuEntry(Actions.AddCurveToGroup);
@@ -90,6 +211,26 @@ namespace UE::NeuralMorphModel
 			{			
 				Menu.AddMenuEntry(Actions.DeleteSelectedItems);
 			}
+
+			if (!GetRootElements().IsEmpty())
+			{
+				Menu.AddMenuEntry(Actions.ClearGroups);
+			}
+		Menu.EndSection();
+
+		Menu.BeginSection("CopyPasteActions", LOCTEXT("CopyPasteActionsHeading", "Copy/Paste Actions"));
+		{
+			if (!RootElements.IsEmpty())
+			{
+				Menu.AddMenuEntry(Actions.CopyCurveGroups);
+			}
+
+			TSharedPtr<FJsonObject> JsonObject;
+			if (InputWidget->HasValidClipBoardData(TEXT("CurveGroups"), JsonObject))
+			{
+				Menu.AddMenuEntry(Actions.PasteCurveGroups);
+			}
+		}
 		Menu.EndSection();
 
 		return Menu.MakeWidget();
@@ -134,6 +275,7 @@ namespace UE::NeuralMorphModel
 			CurveGroup.CurveNames.Add(CurveName);
 		}
 
+		EditorModel->UpdateEditorInputInfo();
 		RefreshTree(true);
 		EditorModel->RebuildEditorMaskInfo();
 	}
@@ -146,7 +288,7 @@ namespace UE::NeuralMorphModel
 		TArray<TSharedPtr<FNeuralMorphCurveGroupsTreeElement>> SelectedGroups;
 		TArray<TSharedPtr<FNeuralMorphCurveGroupsTreeElement>> SelectedCurves;
 
-		// Check if the selection contains bones and/or groups.
+		// Check if the selection contains curves and/or groups.
 		for (const TSharedPtr<FNeuralMorphCurveGroupsTreeElement>& SelectedItem : CurSelectedItems)
 		{
 			if (SelectedItem->IsGroup())
@@ -165,7 +307,7 @@ namespace UE::NeuralMorphModel
 		FScopedTransaction Transaction(TEXT("SNeuralMorphCurveGroupsWidget"), LOCTEXT("RemoveCurveGroupItemsText", "Remove Curve Group Items"), NeuralMorphModel);
 		NeuralMorphModel->Modify();
 
-		// Remove all selected bones.
+		// Remove all selected curves.
 		if (!SelectedCurves.IsEmpty())
 		{
 			for (const TSharedPtr<FNeuralMorphCurveGroupsTreeElement>& CurveItem : SelectedCurves)
@@ -203,6 +345,7 @@ namespace UE::NeuralMorphModel
 			}
 		}
 
+		EditorModel->UpdateEditorInputInfo();
 		RefreshTree(true);
 		EditorModel->RebuildEditorMaskInfo();
 	}
@@ -215,6 +358,7 @@ namespace UE::NeuralMorphModel
 		NeuralMorphModel->Modify();
 
 		NeuralMorphModel->CurveGroups.Empty();
+		EditorModel->UpdateEditorInputInfo();
 		RefreshTree(true);
 		EditorModel->RebuildEditorMaskInfo();
 	}
@@ -268,6 +412,7 @@ namespace UE::NeuralMorphModel
 				CurveGroup.CurveNames.Add(CurveName);
 			}
 
+			EditorModel->UpdateEditorInputInfo();
 			RefreshTree(true);
 			EditorModel->RebuildEditorMaskInfo();
 		}
@@ -454,9 +599,9 @@ namespace UE::NeuralMorphModel
 		return SectionTitle;
 	}
 
-	void SNeuralMorphCurveGroupsWidget::Refresh()
+	void SNeuralMorphCurveGroupsWidget::Refresh(bool bBroadcastPropertyChanged)
 	{ 
-		RefreshTree(false);
+		RefreshTree(bBroadcastPropertyChanged);
 	}
 	
 	TSharedPtr<UE::MLDeformer::SMLDeformerInputWidget> SNeuralMorphCurveGroupsWidget::GetInputWidget() const
diff --git a/Engine/Plugins/Animation/MLDeformer/NeuralMorphModel/Source/NeuralMorphModelEditor/Private/SNeuralMorphInputWidget.cpp b/Engine/Plugins/Animation/MLDeformer/NeuralMorphModel/Source/NeuralMorphModelEditor/Private/SNeuralMorphInputWidget.cpp
--- a/Engine/Plugins/Animation/MLDeformer/NeuralMorphModel/Source/NeuralMorphModelEditor/Private/SNeuralMorphInputWidget.cpp
+++ b/Engine/Plugins/Animation/MLDeformer/NeuralMorphModel/Source/NeuralMorphModelEditor/Private/SNeuralMorphInputWidget.cpp
@@ -618,15 +618,22 @@ namespace UE::NeuralMorphModel
 		BoneGroupsCommandList = MakeShared<FUICommandList>();
 
 		FNeuralMorphEditorModel* NeuralEditorModel = static_cast<FNeuralMorphEditorModel*>(EditorModel);
-		SAssignNew(BoneGroupsWidget, SNeuralMorphBoneGroupsWidget)
-			.EditorModel(NeuralEditorModel)
-			.InputWidget(SharedThis(this));
+
+		const TSharedPtr<SWidget> MainWidget =
+			SNew(SVerticalBox)
+			+SVerticalBox::Slot()
+			.MinHeight(75.0f)
+			[
+				SAssignNew(BoneGroupsWidget, SNeuralMorphBoneGroupsWidget)
+					.EditorModel(NeuralEditorModel)
+					.InputWidget(SharedThis(this))
+			];
 
 		FSectionInfo SectionInfo;
 		SectionInfo.SectionTitle = TAttribute<FText>::CreateSP(BoneGroupsWidget.Get(), &SNeuralMorphBoneGroupsWidget::GetSectionTitle);
 		SectionInfo.PlusButtonPressed = FOnClicked::CreateSP(this, &SNeuralMorphInputWidget::ShowBoneGroupsManageContextMenu);
 		SectionInfo.PlusButtonTooltip = LOCTEXT("BonesGroupsPlusButtonTooltip", "Manage bone groups.");
-		AddSection(BoneGroupsWidget, SectionInfo);
+		AddSection(MainWidget, SectionInfo);
 
 		BoneGroupsWidget->BindCommands(BoneGroupsCommandList);
 	}
@@ -666,15 +673,22 @@ namespace UE::NeuralMorphModel
 		CurveGroupsCommandList = MakeShared<FUICommandList>();
 
 		FNeuralMorphEditorModel* NeuralEditorModel = static_cast<FNeuralMorphEditorModel*>(EditorModel);
-		SAssignNew(CurveGroupsWidget, SNeuralMorphCurveGroupsWidget)
-			.EditorModel(NeuralEditorModel)
-			.InputWidget(SharedThis(this));
+
+		const TSharedPtr<SWidget> MainWidget =
+			SNew(SVerticalBox)
+			+SVerticalBox::Slot()
+			.MinHeight(75.0f)
+			[
+				SAssignNew(CurveGroupsWidget, SNeuralMorphCurveGroupsWidget)
+					.EditorModel(NeuralEditorModel)
+					.InputWidget(SharedThis(this))
+			];
 
 		FSectionInfo SectionInfo;
 		SectionInfo.SectionTitle = TAttribute<FText>::CreateSP(CurveGroupsWidget.Get(), &SNeuralMorphCurveGroupsWidget::GetSectionTitle);
 		SectionInfo.PlusButtonPressed = FOnClicked::CreateSP(this, &SNeuralMorphInputWidget::ShowCurveGroupsManageContextMenu);
 		SectionInfo.PlusButtonTooltip = LOCTEXT("CurveGroupsPlusButtonTooltip", "Manage curve groups.");
-		AddSection(CurveGroupsWidget, SectionInfo);
+		AddSection(MainWidget, SectionInfo);
 
 		CurveGroupsWidget->BindCommands(CurveGroupsCommandList);
 	}
@@ -839,13 +853,20 @@ namespace UE::NeuralMorphModel
 			check(InputBoneIndex != INDEX_NONE);
 			NeuralEditorModel->GenerateBoneMaskInfo(InputBoneIndex, HierarchyDepth);
 		}
+		if (BoneGroupsWidget.IsValid())
+		{
+			BoneGroupsWidget->Refresh(true);
+		}
 		NeuralEditorModel->RebuildEditorMaskInfo();
 	}
 
 	void SNeuralMorphInputWidget::OnAddInputCurves(const TArray<FName>& Names)
 	{
 		FNeuralMorphEditorModel* NeuralEditorModel = static_cast<FNeuralMorphEditorModel*>(EditorModel);
-		UNeuralMorphModel* NeuralMorphModel = NeuralEditorModel->GetNeuralMorphModel();
+		if (CurveGroupsWidget.IsValid())
+		{
+			CurveGroupsWidget->Refresh(true);
+		}
 		NeuralEditorModel->RebuildEditorMaskInfo();
 	}
 
diff --git a/Engine/Plugins/Animation/MLDeformer/NeuralMorphModel/Source/NeuralMorphModelEditor/Public/SNeuralMorphBoneGroupsWidget.h b/Engine/Plugins/Animation/MLDeformer/NeuralMorphModel/Source/NeuralMorphModelEditor/Public/SNeuralMorphBoneGroupsWidget.h
--- a/Engine/Plugins/Animation/MLDeformer/NeuralMorphModel/Source/NeuralMorphModelEditor/Public/SNeuralMorphBoneGroupsWidget.h
+++ b/Engine/Plugins/Animation/MLDeformer/NeuralMorphModel/Source/NeuralMorphModelEditor/Public/SNeuralMorphBoneGroupsWidget.h
@@ -34,6 +34,8 @@ namespace UE::NeuralMorphModel
 		TSharedPtr<FUICommandInfo> DeleteSelectedItems;
 		TSharedPtr<FUICommandInfo> ClearGroups;
 		TSharedPtr<FUICommandInfo> AddBoneToGroup;
+		TSharedPtr<FUICommandInfo> CopyBoneGroups;
+		TSharedPtr<FUICommandInfo> PasteBoneGroups;
 	};
 
 
@@ -80,7 +82,7 @@ namespace UE::NeuralMorphModel
 
 		UE_API void Construct(const FArguments& InArgs);
 		UE_API void BindCommands(TSharedPtr<FUICommandList> CommandList);
-		UE_API void Refresh();
+		UE_API void Refresh(bool bBroadcastPropertyChanged=false);
 
 		UE_API FText GetSectionTitle() const;
 		UE_API int32 GetNumSelectedGroups() const;
@@ -91,7 +93,7 @@ namespace UE::NeuralMorphModel
 		UE_API void HandleGetChildrenForTree(TSharedPtr<FNeuralMorphBoneGroupsTreeElement> InItem, TArray<TSharedPtr<FNeuralMorphBoneGroupsTreeElement>>& OutChildren);	
 		UE_API void OnSelectionChanged(TSharedPtr<FNeuralMorphBoneGroupsTreeElement> Selection, ESelectInfo::Type SelectInfo);
 		UE_API const TArray<TSharedPtr<FNeuralMorphBoneGroupsTreeElement>>& GetRootElements() const;
-		UE_API FReply OnKeyDown(const FGeometry& InGeometry, const FKeyEvent& InKeyEvent) override;
+		UE_API virtual FReply OnKeyDown(const FGeometry& InGeometry, const FKeyEvent& InKeyEvent) override;
 		UE_API TSharedRef<ITableRow> MakeTableRowWidget(TSharedPtr<FNeuralMorphBoneGroupsTreeElement> InItem, const TSharedRef<STableViewBase>& OwnerTable);
 		UE_API TSharedPtr<SWidget> CreateContextMenuWidget() const;
 
@@ -100,6 +102,9 @@ namespace UE::NeuralMorphModel
 		UE_API bool BroadcastModelPropertyChanged(const FName PropertyName);
 		UE_API TSharedPtr<SWidget> CreateContextWidget() const;
 
+		UE_API void OnCopyBoneGroups() const;
+		UE_API void OnPasteBoneGroups();
+	
 		UE_API void OnCreateBoneGroup();
 		UE_API void OnDeleteSelectedItems();
 		UE_API void OnClearBoneGroups();
diff --git a/Engine/Plugins/Animation/MLDeformer/NeuralMorphModel/Source/NeuralMorphModelEditor/Public/SNeuralMorphCurveGroupsWidget.h b/Engine/Plugins/Animation/MLDeformer/NeuralMorphModel/Source/NeuralMorphModelEditor/Public/SNeuralMorphCurveGroupsWidget.h
--- a/Engine/Plugins/Animation/MLDeformer/NeuralMorphModel/Source/NeuralMorphModelEditor/Public/SNeuralMorphCurveGroupsWidget.h
+++ b/Engine/Plugins/Animation/MLDeformer/NeuralMorphModel/Source/NeuralMorphModelEditor/Public/SNeuralMorphCurveGroupsWidget.h
@@ -35,6 +35,8 @@ namespace UE::NeuralMorphModel
 		TSharedPtr<FUICommandInfo> DeleteSelectedItems;
 		TSharedPtr<FUICommandInfo> ClearGroups;
 		TSharedPtr<FUICommandInfo> AddCurveToGroup;
+		TSharedPtr<FUICommandInfo> CopyCurveGroups;
+		TSharedPtr<FUICommandInfo> PasteCurveGroups;
 	};
 
 
@@ -81,7 +83,7 @@ namespace UE::NeuralMorphModel
 
 		UE_API void Construct(const FArguments& InArgs);
 		UE_API void BindCommands(TSharedPtr<FUICommandList> CommandList);
-		UE_API void Refresh();
+		UE_API void Refresh(bool bBroadcastPropertyChanged=false);
 
 		UE_API FText GetSectionTitle() const;
 		UE_API int32 GetNumSelectedGroups() const;
@@ -92,7 +94,7 @@ namespace UE::NeuralMorphModel
 		UE_API void HandleGetChildrenForTree(TSharedPtr<FNeuralMorphCurveGroupsTreeElement> InItem, TArray<TSharedPtr<FNeuralMorphCurveGroupsTreeElement>>& OutChildren);	
 		UE_API void OnSelectionChanged(TSharedPtr<FNeuralMorphCurveGroupsTreeElement> Selection, ESelectInfo::Type SelectInfo);
 		UE_API const TArray<TSharedPtr<FNeuralMorphCurveGroupsTreeElement>>& GetRootElements() const;
-		UE_API FReply OnKeyDown(const FGeometry& InGeometry, const FKeyEvent& InKeyEvent) override;
+		UE_API virtual FReply OnKeyDown(const FGeometry& InGeometry, const FKeyEvent& InKeyEvent) override;
 		UE_API TSharedRef<ITableRow> MakeTableRowWidget(TSharedPtr<FNeuralMorphCurveGroupsTreeElement> InItem, const TSharedRef<STableViewBase>& OwnerTable);
 		UE_API TSharedPtr<SWidget> CreateContextMenuWidget() const;
 
@@ -101,6 +103,9 @@ namespace UE::NeuralMorphModel
 		UE_API bool BroadcastModelPropertyChanged(const FName PropertyName);
 		UE_API TSharedPtr<SWidget> CreateContextWidget() const;
 
+		UE_API void OnCopyCurveGroups() const;
+		UE_API void OnPasteCurveGroups();
+		
 		UE_API void OnCreateCurveGroup();
 		UE_API void OnDeleteSelectedItems();
 		UE_API void OnClearCurveGroups();
