[MLDeformer] Fix async loading issues.

diff --git a/Engine/Plugins/Animation/MLDeformer/MLDeformerFramework/Source/MLDeformerFramework/Private/MLDeformerModel.cpp b/Engine/Plugins/Animation/MLDeformer/MLDeformerFramework/Source/MLDeformerFramework/Private/MLDeformerModel.cpp
--- a/Engine/Plugins/Animation/MLDeformer/MLDeformerFramework/Source/MLDeformerFramework/Private/MLDeformerModel.cpp
+++ b/Engine/Plugins/Animation/MLDeformer/MLDeformerFramework/Source/MLDeformerFramework/Private/MLDeformerModel.cpp
@@ -18,6 +18,7 @@
 #include "Components/SkeletalMeshComponent.h"
 #include "UObject/AssetRegistryTagsContext.h"
 #include "Misc/TransactionObjectEvent.h"
+#include "UObject/ObjectSaveContext.h"
 
 #include UE_INLINE_GENERATED_CPP_BY_NAME(MLDeformerModel)
 
@@ -110,19 +111,13 @@ void UMLDeformerModel::Serialize(FArchive& Archive)
 
 	Super::Serialize(Archive);
 	Archive.UsingCustomVersion(UE::MLDeformer::FMLDeformerObjectVersion::GUID);
+}
 
-	#if WITH_EDITOR
-		if (Archive.IsSaving())
-		{
-			InitVertexMap();
-		}
-
-		if (!Archive.IsCooking() && Archive.IsLoading())
-		{
-			// This also triggers the target mesh to be loaded, don't do that while cooking.
-			UpdateCachedNumVertices();
-		}
-	#endif
+void UMLDeformerModel::PreSave(FObjectPreSaveContext SaveContext)
+{
+#if WITH_EDITOR
+	InitVertexMap();
+#endif	
 }
 
 UMLDeformerAsset* UMLDeformerModel::GetDeformerAsset() const
@@ -135,9 +130,10 @@ void UMLDeformerModel::PostLoad()
 	TRACE_CPUPROFILER_EVENT_SCOPE(UMLDeformerModel::PostLoad)
 
 	Super::PostLoad();
-
+	
 	#if WITH_EDITOR
 		InitVertexMap();
+		UpdateCachedNumVertices();
 	#endif
 
 	InitGPUData();
diff --git a/Engine/Plugins/Animation/MLDeformer/MLDeformerFramework/Source/MLDeformerFramework/Public/MLDeformerModel.h b/Engine/Plugins/Animation/MLDeformer/MLDeformerFramework/Source/MLDeformerFramework/Public/MLDeformerModel.h
--- a/Engine/Plugins/Animation/MLDeformer/MLDeformerFramework/Source/MLDeformerFramework/Public/MLDeformerModel.h
+++ b/Engine/Plugins/Animation/MLDeformer/MLDeformerFramework/Source/MLDeformerFramework/Public/MLDeformerModel.h
@@ -356,6 +356,7 @@ public:
 
 	// UObject overrides.
 	UE_API virtual void Serialize(FArchive& Archive) override;
+	UE_API virtual void PreSave(FObjectPreSaveContext SaveContext) override;
 	UE_API virtual void PostLoad() override;
 	UE_API virtual void BeginDestroy() override;
 	UE_API virtual bool IsReadyForFinishDestroy() override;
diff --git a/Engine/Plugins/Animation/MLDeformer/NeuralMorphModel/Source/NeuralMorphModel/Private/NeuralMorphModel.cpp b/Engine/Plugins/Animation/MLDeformer/NeuralMorphModel/Source/NeuralMorphModel/Private/NeuralMorphModel.cpp
--- a/Engine/Plugins/Animation/MLDeformer/NeuralMorphModel/Source/NeuralMorphModel/Private/NeuralMorphModel.cpp
+++ b/Engine/Plugins/Animation/MLDeformer/NeuralMorphModel/Source/NeuralMorphModel/Private/NeuralMorphModel.cpp
@@ -224,21 +224,6 @@ void UNeuralMorphModel::Serialize(FArchive& Archive)
 				BoneGroupMaskInfos_DEPRECATED.Empty();
 			PRAGMA_ENABLE_DEPRECATION_WARNINGS;
 		}
-
-		// Since 5.4 (when LOD was added), we modified the way the bone masks are generated.
-		// In UE 5.3 we added some extra bones to the mask. To keep the same behavior, we will upgrade the mask to include those bones again to ensure we get the same results when training.
-		// Newer assets will not do this.
-		if (Archive.IsLoading() && 
-			Archive.CustomVer(UE::MLDeformer::FMLDeformerObjectVersion::GUID) < UE::MLDeformer::FMLDeformerObjectVersion::LODSupportAdded)
-		{
-			USkeletalMesh* SkelMesh = GetSkeletalMesh();
-			if (SkelMesh)
-			{
-				const FReferenceSkeleton& RefSkeleton = SkelMesh->GetRefSkeleton();
-				ModifyMaskInfosForBackwardCompatibility(RefSkeleton, GetBoneIncludeList(), BoneMaskInfoMap);
-				ModifyMaskInfosForBackwardCompatibility(RefSkeleton, GetBoneIncludeList(), BoneGroupMaskInfoMap);
-			}
-		}
 	#endif
 }
 
@@ -268,6 +253,24 @@ void UNeuralMorphModel::UpdateMissingGroupNames()
 void UNeuralMorphModel::PostLoad()
 {
 	Super::PostLoad();
+	
+#if WITH_EDITORONLY_DATA
+	// Since 5.4 (when LOD was added), we modified the way the bone masks are generated.
+	// In UE 5.3 we added some extra bones to the mask. To keep the same behavior, we will upgrade the mask to include those bones again to ensure we get the same results when training.
+	// Newer assets will not do this.
+	const int32 Version = GetLinkerCustomVersion(UE::MLDeformer::FMLDeformerObjectVersion::GUID);
+	if (Version < UE::MLDeformer::FMLDeformerObjectVersion::LODSupportAdded)
+	{
+		USkeletalMesh* SkelMesh = GetSkeletalMesh();
+		if (SkelMesh)
+		{
+			const FReferenceSkeleton& RefSkeleton = SkelMesh->GetRefSkeleton();
+			ModifyMaskInfosForBackwardCompatibility(RefSkeleton, GetBoneIncludeList(), BoneMaskInfoMap);
+			ModifyMaskInfosForBackwardCompatibility(RefSkeleton, GetBoneIncludeList(), BoneGroupMaskInfoMap);
+		}
+	}
+#endif
+	
 	UpdateMissingGroupNames();
 	check(SupportsGlobalModeOnly() ? Mode == ENeuralMorphMode::Global : true);
 }
